Add the "Game Input" and "Game Input Windows" plugins into the Engine folder.

Game Input is a next-gen input API

diff --git a/Engine/Plugins/Experimental/GameInputWindows/GameInputWindows.uplugin b/Engine/Plugins/Experimental/GameInputWindows/GameInputWindows.uplugin
new file mode 100644
--- /dev/null
+++ b/Engine/Plugins/Experimental/GameInputWindows/GameInputWindows.uplugin
@@ -0,0 +1,36 @@
+{
+	"FileVersion": 3,
+	"FriendlyName": "Game Input (Windows)",
+	"Version": 1,
+	"VersionName": "1.0",
+	"CreatedBy": "Epic Games, Inc.",
+	"CreatedByURL": "https://epicgames.com",
+	"Description": "GameInput is a next-generation input API that exposes input devices of all kinds through a single consistent interface.",
+	"Category": "Input Devices",
+	"EnabledByDefault": false,
+	"IsBetaVersion": true,
+	"CanContainContent": false,
+	"HasExplicitPlatforms": true,
+	"bRequiresBuildPlatform": false,
+	"SupportedTargetPlatforms": [
+		"Win64"
+	],
+	"Modules": [
+		{
+			"Name": "GameInputWindows",
+			"Type": "Runtime",
+			"HasExplicitPlatforms": true,
+			"TargetDenyList": [ "Server" ],
+			"PlatformAllowList": [
+				"Win64"
+			]
+		}
+	],
+	"Plugins": [
+		{
+			"Name": "GameInput",
+			"Enabled": true,
+			"TargetDenyList": [ "Server" ]
+		}
+	]
+}
\ No newline at end of file
diff --git a/Engine/Plugins/Experimental/GameInputWindows/Source/GameInputWindows/GameInputWindows.build.cs b/Engine/Plugins/Experimental/GameInputWindows/Source/GameInputWindows/GameInputWindows.build.cs
new file mode 100644
--- /dev/null
+++ b/Engine/Plugins/Experimental/GameInputWindows/Source/GameInputWindows/GameInputWindows.build.cs
@@ -0,0 +1,26 @@
+// Copyright Epic Games, Inc. All Rights Reserved.
+
+using System;
+
+namespace UnrealBuildTool.Rules
+{
+	public class GameInputWindows : ModuleRules
+	{
+		public GameInputWindows(ReadOnlyTargetRules Target) : base(Target)
+		{
+			PublicDependencyModuleNames.AddRange(
+				new string[]
+				{
+					"Core",
+                    "ApplicationCore",
+                    "Engine",
+                    "InputCore",
+                    "InputDevice",				
+                    "CoreUObject",
+                    "DeveloperSettings",
+					"GameInputBase",
+				}
+			);
+		}
+	}
+}
\ No newline at end of file
diff --git a/Engine/Plugins/Experimental/GameInputWindows/Source/GameInputWindows/Private/GameInputWindowsDevice.cpp b/Engine/Plugins/Experimental/GameInputWindows/Source/GameInputWindows/Private/GameInputWindowsDevice.cpp
new file mode 100644
--- /dev/null
+++ b/Engine/Plugins/Experimental/GameInputWindows/Source/GameInputWindows/Private/GameInputWindowsDevice.cpp
@@ -0,0 +1,134 @@
+// Copyright Epic Games, Inc. All Rights Reserved.
+
+#include "GameInputWindowsDevice.h"
+
+#if GAME_INPUT_SUPPORT
+
+#include "GameInputUtils.h"
+#include "GameInputLogging.h"
+#include "HAL/IConsoleManager.h" // FAutoConsoleVariableRef
+#include "GenericPlatform/IInputInterface.h"
+#include "GameInputDeveloperSettings.h"
+
+FGameInputWindowsInputDevice::FGameInputWindowsInputDevice(const TSharedRef<FGenericApplicationMessageHandler>& InMessageHandler, IGameInput* InGameInput)
+	: IGameInputDeviceInterface(InMessageHandler, InGameInput)
+{
+	
+}
+
+FGameInputWindowsInputDevice::~FGameInputWindowsInputDevice()
+{
+	
+}
+
+GameInputKind FGameInputWindowsInputDevice::GetCurrentGameInputKindSupport() const
+{
+	GameInputKind RegisterInputKindMask = IGameInputDeviceInterface::GetCurrentGameInputKindSupport();
+
+	// For now, we want to explicitly make sure that Keyboard and Mouse are NOT being processed by GameInput on windows.
+	// This is because the WindowsApplication is doing to do all the processing of these types for us already, and 
+	// having the Game Input device plugin process them too would result in "double" events and mouse accumulation.
+
+#if !UE_BUILD_SHIPPING
+	static bool bLogOnce = false;
+	if (!bLogOnce)
+	{
+		if ((RegisterInputKindMask & GameInputKindKeyboard) != 0)
+		{
+			UE_LOG(LogGameInput, Log, TEXT("[FGameInputWindowsInputDevice::GetCurrentGameInputKindSupport] Keyboard support was requested, but is not currently supported via the GameInput plugin on Windows."));
+		}
+
+		if ((RegisterInputKindMask & GameInputKindMouse) != 0)
+		{
+			UE_LOG(LogGameInput, Log, TEXT("[FGameInputWindowsInputDevice::GetCurrentGameInputKindSupport] Mouse support was requested, but is not currently supported via the GameInput plugin on Windows."));
+		}
+		bLogOnce = true;
+	}
+#endif	// !UE_BUILD_SHIPPING
+
+	// Clear the bits on these types to make sure we don't process any of them
+	RegisterInputKindMask &= ~GameInputKindKeyboard;
+	RegisterInputKindMask &= ~GameInputKindMouse;
+
+	return RegisterInputKindMask;
+}
+
+void FGameInputWindowsInputDevice::HandleDeviceDisconnected(IGameInputDevice* Device, uint64 Timestamp)
+{	
+	if (Device)
+	{
+		if (FGameInputDeviceContainer* Data = GetDeviceData(Device))
+		{
+			// Clear any input state that might be related to this device
+			Data->ClearInputState(GameInput);
+
+			// Set it's device to nullptr because it is now disconnected
+			Data->SetGameInputDevice(nullptr);
+			UE_LOG(LogGameInput, Log, TEXT("Game Input Device '%s' Disconnected Successfully at Input Device ID '%d'"), *UE::GameInput::LexToString(Device), Data->GetDeviceId().GetId());
+
+			// Remap this device to the "unpaired" user because it has been disconnected
+			const FPlatformUserId NewUserToAssign = IPlatformInputDeviceMapper::Get().GetUserForUnpairedInputDevices();
+			const FInputDeviceId DeviceId = Data->GetDeviceId();
+
+			const bool bSuccess = IPlatformInputDeviceMapper::Get().Internal_MapInputDeviceToUser(DeviceId, NewUserToAssign, EInputDeviceConnectionState::Disconnected);
+			if (bSuccess)
+			{
+				Data->SetPlatformUserId(NewUserToAssign);
+			}
+		}
+		else
+		{
+			UE_LOG(LogGameInput, Error, TEXT("Game Input failed to disconnect a device. The Device '%s' did not have an associated FGameInputWindowsInputDevice!"), *UE::GameInput::LexToString(Device));
+		}
+	}
+	else
+	{
+		UE_LOG(LogGameInput, Warning, TEXT("Game Input failed to disconnect a device. The Device was null! %s"), *UE::GameInput::LexToString(Device));
+	}
+
+	EnumerateCurrentlyConnectedDeviceTypes();
+}
+
+void FGameInputWindowsInputDevice::HandleDeviceConnected(IGameInputDevice* Device, uint64 Timestamp)
+{
+	// get device information
+	const GameInputDeviceInfo* Info = Device->GetDeviceInfo();
+	UE_LOG(LogGameInput, Log, TEXT("Game Input Device Connected: %s  of kind: %s"), *UE::GameInput::LexToString(Device), *UE::GameInput::LexToString(Info->supportedInput));
+	
+	FGameInputDeviceContainer* Data = GetOrCreateDeviceData(Device);
+	check(Data);
+
+	// Map this input device to its user	
+	const FInputDeviceId DeviceId = Data->GetDeviceId();
+
+	// TODO: This will map every game input device to the "primary" user. This is NOT what we want in the long term.
+	// Long term, we want you to be able to configure how input devices are mapped to local players in UE with some kind of
+	// schema or config variable. We should add this as a setting in the Developer Settings object here.
+	const FPlatformUserId UserToAssign = IPlatformInputDeviceMapper::Get().GetPrimaryPlatformUser();
+	
+	const bool bSuccess = IPlatformInputDeviceMapper::Get().Internal_MapInputDeviceToUser(DeviceId, UserToAssign, EInputDeviceConnectionState::Connected);
+	if (ensure(bSuccess))
+	{
+		Data->SetPlatformUserId(UserToAssign);
+	}
+
+	UE_LOG(LogGameInput, Log, TEXT("Using PlatformUserId %d and InputDeviceId %d for device %s"), UserToAssign.GetInternalId(), DeviceId.GetId(), *UE::GameInput::LexToString(Device));
+
+	EnumerateCurrentlyConnectedDeviceTypes();
+}
+
+FGameInputDeviceContainer* FGameInputWindowsInputDevice::CreateDeviceData(IGameInputDevice* InDevice)
+{
+	TSharedPtr<FGameInputDeviceContainer> Container = MakeShared<FGameInputDeviceContainer>(MessageHandler, InDevice, GetCurrentGameInputKindSupport());
+	Container->InitalizeDeviceProcessors();
+
+	// This is a new device, we need to assign a new input device ID from the platform user
+	const FInputDeviceId AssignedInputDeviceId = IPlatformInputDeviceMapper::Get().AllocateNewInputDeviceId();
+	Container->SetInputDeviceId(AssignedInputDeviceId);
+
+	DeviceData.Emplace(Container);
+
+	return Container.Get();
+}
+
+#endif	//#if GAME_INPUT_SUPPORT
\ No newline at end of file
diff --git a/Engine/Plugins/Experimental/GameInputWindows/Source/GameInputWindows/Private/GameInputWindowsDevice.h b/Engine/Plugins/Experimental/GameInputWindows/Source/GameInputWindows/Private/GameInputWindowsDevice.h
new file mode 100644
--- /dev/null
+++ b/Engine/Plugins/Experimental/GameInputWindows/Source/GameInputWindows/Private/GameInputWindowsDevice.h
@@ -0,0 +1,25 @@
+// Copyright Epic Games, Inc. All Rights Reserved.
+
+#pragma once
+
+#if GAME_INPUT_SUPPORT
+
+#include "IGameInputDeviceInterface.h"
+
+class FGameInputWindowsInputDevice final : public IGameInputDeviceInterface
+{
+public:
+	explicit FGameInputWindowsInputDevice(const TSharedRef<FGenericApplicationMessageHandler>& InMessageHandler, struct IGameInput* InGameInput);
+	virtual ~FGameInputWindowsInputDevice() override;
+
+protected:
+
+	//~Begin IGameInputDeviceInterface interface	
+	virtual GameInputKind GetCurrentGameInputKindSupport() const override;
+	virtual void HandleDeviceDisconnected(IGameInputDevice* Device, uint64 Timestamp) override;
+	virtual void HandleDeviceConnected(IGameInputDevice* Device, uint64 Timestamp) override;
+	virtual FGameInputDeviceContainer* CreateDeviceData(IGameInputDevice* InDevice) override;
+	//~End IGameInputDeviceInterface interface
+};
+
+#endif	// #if GAME_INPUT_SUPPORT
\ No newline at end of file
diff --git a/Engine/Plugins/Experimental/GameInputWindows/Source/GameInputWindows/Private/GameInputWindowsModule.cpp b/Engine/Plugins/Experimental/GameInputWindows/Source/GameInputWindows/Private/GameInputWindowsModule.cpp
new file mode 100644
--- /dev/null
+++ b/Engine/Plugins/Experimental/GameInputWindows/Source/GameInputWindows/Private/GameInputWindowsModule.cpp
@@ -0,0 +1,78 @@
+// Copyright Epic Games, Inc. All Rights Reserved.
+
+#include "GameInputWindowsModule.h"
+
+#include "CoreGlobals.h"
+#include "GameInputWindowsDevice.h"
+#include "GameInputBaseModule.h"
+#include "GameInputLogging.h"
+#include "Misc/App.h"
+#include "Modules/ModuleManager.h"
+
+#if GAME_INPUT_SUPPORT
+THIRD_PARTY_INCLUDES_START
+#include "Microsoft/AllowMicrosoftPlatformTypes.h"
+#include <GameInput.h>
+#include "Microsoft/HideMicrosoftPlatformTypes.h"
+THIRD_PARTY_INCLUDES_END
+#endif	// GAME_INPUT_SUPPORT
+
+namespace UE::GameInput
+{
+	/** The name of this modular feature plugin. */
+	static const FName GameInputWindowsFeatureName = TEXT("GameInputWindows");
+}
+
+FGameInputWindowsModule& FGameInputWindowsModule::Get()
+{
+	return FModuleManager::LoadModuleChecked<FGameInputWindowsModule>(UE::GameInput::GameInputWindowsFeatureName);
+}
+
+bool FGameInputWindowsModule::IsAvailable()
+{
+	return FModuleManager::Get().IsModuleLoaded(UE::GameInput::GameInputWindowsFeatureName);
+}
+
+TSharedPtr<IInputDevice> FGameInputWindowsModule::CreateInputDevice(const TSharedRef<FGenericApplicationMessageHandler>& InMessageHandler)
+{
+#if GAME_INPUT_SUPPORT
+	if (IGameInput* GameInputPtr = FGameInputBaseModule::GetGameInput())
+	{
+		TSharedPtr<FGameInputWindowsInputDevice> DevicePtr = MakeShared<FGameInputWindowsInputDevice>(InMessageHandler, GameInputPtr);
+		DevicePtr->Initialize();
+		return DevicePtr;
+	}
+	else
+	{
+		if (IsRunningCommandlet())
+		{
+			// We don't care for Game Input if we are running a commandlet, like when we are cooking or running C++ automation tests
+			UE_LOG(LogGameInput, Log, TEXT("[%s] GameInput interface will not be created because IsRunningCommandlet is true."), ANSI_TO_TCHAR(__FUNCTION__));
+		}
+		else if (!FApp::HasProjectName())
+		{
+			// If there is no project name then we don't need game input either. This means we are in the project launcher
+			UE_LOG(LogGameInput, Log, TEXT("[%s] GameInput interface will not be created because there is no project name."), ANSI_TO_TCHAR(__FUNCTION__));
+		}
+		else if (FApp::IsUnattended())
+		{
+			UE_LOG(LogGameInput, Log, TEXT("[%s] GameInput interface will not be created because FApp::IsUnattended is true!"), ANSI_TO_TCHAR(__FUNCTION__));
+		}
+		else if (!FApp::CanEverRender())
+		{
+			UE_LOG(LogGameInput, Log, TEXT("[%s] GameInput interface will not be created because FApp::CanEverRender is false!"), ANSI_TO_TCHAR(__FUNCTION__));
+		}
+		else
+		{
+			// Otherwise, we want Game Input to be properly initialized
+			UE_LOG(LogGameInput, Log, TEXT("[%s] Failed to get the IGameInput object from the base module! You will not be able to process input."), ANSI_TO_TCHAR(__FUNCTION__));	
+		}
+	}
+#else
+	UE_LOG(LogGameInput, Error, TEXT("[FGameInputWindowsModule] Failed to create a GameInput device! GAME_INPUT_SUPPORT is false!"));
+#endif	// GAME_INPUT_SUPPORT
+
+	return nullptr;
+}
+
+IMPLEMENT_MODULE(FGameInputWindowsModule, GameInputWindows)
\ No newline at end of file
diff --git a/Engine/Plugins/Experimental/GameInputWindows/Source/GameInputWindows/Private/GameInputWindowsModule.h b/Engine/Plugins/Experimental/GameInputWindows/Source/GameInputWindows/Private/GameInputWindowsModule.h
new file mode 100644
--- /dev/null
+++ b/Engine/Plugins/Experimental/GameInputWindows/Source/GameInputWindows/Private/GameInputWindowsModule.h
@@ -0,0 +1,21 @@
+// Copyright Epic Games, Inc. All Rights Reserved.
+
+#pragma once
+
+#include "IInputDeviceModule.h"
+
+/**
+ * Input Device module that will create the Game Input Device to handle input on Windows platforms.
+ */
+class FGameInputWindowsModule : public IInputDeviceModule
+{
+public:
+	static FGameInputWindowsModule& Get();
+
+	/** Returns true if this module is loaded (aka available) by the FModuleManager */
+	static bool IsAvailable();
+
+	//~ Begin IInputDeviceModule interface
+	virtual TSharedPtr<IInputDevice> CreateInputDevice(const TSharedRef<FGenericApplicationMessageHandler>& InMessageHandler) override;		
+	//~ End IInputDeviceModule interface
+};
\ No newline at end of file
diff --git a/Engine/Plugins/Runtime/GameInput/GameInput.uplugin b/Engine/Plugins/Runtime/GameInput/GameInput.uplugin
new file mode 100644
--- /dev/null
+++ b/Engine/Plugins/Runtime/GameInput/GameInput.uplugin
@@ -0,0 +1,30 @@
+{
+	"FileVersion": 3,
+	"FriendlyName": "Game Input Base",
+	"Version": 1,
+	"VersionName": "1.0",
+	"CreatedBy": "Epic Games, Inc.",
+	"CreatedByURL": "https://epicgames.com",
+	"Description": "GameInput is a next-generation input API that exposes input devices of all kinds through a single consistent interface.",
+	"Category": "Input Devices",
+	"EnabledByDefault": false,
+	"IsBetaVersion": false,
+	"CanContainContent": false,
+	"RequiresBuildPlatform": false,
+	"HasExplicitPlatforms": true,
+	"SupportedTargetPlatforms": [ "Win64" ],
+	"Modules": [
+		{
+			"Name": "GameInputBase",
+			"Type": "Runtime",
+			"LoadingPhase": "Default",
+			"HasExplicitPlatforms": true,
+			"PlatformAllowList": [ "Win64" ]
+		},
+		{
+			"Name": "GameInputBaseEditor",
+			"Type": "Editor",
+			"LoadingPhase": "Default"
+		}
+	]
+}
\ No newline at end of file
diff --git a/Engine/Plugins/Runtime/GameInput/Source/GameInputBase/GameInputBase.build.cs b/Engine/Plugins/Runtime/GameInput/Source/GameInputBase/GameInputBase.build.cs
new file mode 100644
--- /dev/null
+++ b/Engine/Plugins/Runtime/GameInput/Source/GameInputBase/GameInputBase.build.cs
@@ -0,0 +1,187 @@
+// Copyright Epic Games, Inc. All Rights Reserved.
+
+using System;
+using System.IO;
+using System.Text.RegularExpressions;
+
+namespace UnrealBuildTool.Rules
+{
+	public class GameInputBase : ModuleRules
+	{
+		private const int GameInputMinimumVersion_Win64 = 220604;
+		
+		/// <summary>
+		/// True if this platform has support for the Game Input library.
+		/// 
+		/// Overriden per-platform implement of the Game Input Base module.
+		/// </summary>
+		protected virtual bool HasGameInputSupport(ReadOnlyTargetRules Target)
+		{
+			// Console platforms will override this function and determine if they have GameInput support on their own.
+			// For the base functionality, we can only support Game Input on windows platforms.
+			if (!Target.Platform.IsInGroup(UnrealPlatformGroup.Windows))
+			{
+				return false;
+			}
+
+			// The Game Input SDK is installed when you install the GRDK, which will set the environment variable "GRDKLatest"
+			// We can then use that to look up the path to the GameInput.h and GameInput.lib files needed to use Game Input.
+			//
+			// https://learn.microsoft.com/en-us/gaming/gdk/_content/gc/input/overviews/input-overview
+			string GRDKLatestPath = Environment.GetEnvironmentVariable("GRDKLatest");
+
+			// We want to have the version number available to us in C++
+			// so that we can validate it at runtime
+			string GDRKVersion = GetGameInputVersionString(GRDKLatestPath);
+			int GRDKIntVersion = -1;
+
+			bool bParseSuccessful = int.TryParse(GDRKVersion, out GRDKIntVersion);
+
+			// We have Game Input support if we could successfully get the version and it is equal to or greater than our min version
+			return bParseSuccessful && GRDKIntVersion >= GameInputMinimumVersion_Win64;
+		}
+
+		/// <summary>
+		/// An extension point for subclasses of this module to add any additional required include/library files
+		/// that may be necessary to add the GameInput SDK to their platform. This will only be called if
+		/// HasGameInputSupport returns true.
+		/// </summary>
+		protected virtual void AddRequiredDeps(ReadOnlyTargetRules Target)
+		{
+			// Console platforms will override this function and determine if they have GameInput support on their own.
+			// For the base functionality, we can only support Game Input on windows platforms.
+			if (!Target.Platform.IsInGroup(UnrealPlatformGroup.Windows))
+			{
+				return;
+			}
+			
+			string GRDKLatestPath = Environment.GetEnvironmentVariable("GRDKLatest");
+			
+			// GameInput.h is from    <GRDKLatest>\GameKit\Include\GameInput.h
+			// GameInput.lib is from  <GRDKLatest>\GameKit\Lib\amd64\GameInput.lib
+			
+			string LibPath = Path.Combine(GRDKLatestPath, "GameKit", "Lib", "amd64", "GameInput.lib");
+			string IncludePath = Path.Combine(GRDKLatestPath, "GameKit", "Include");
+			
+			PublicAdditionalLibraries.Add(LibPath);
+			PublicSystemIncludePaths.Add(IncludePath);
+			
+			// Add this as a delay loaded DLL for cooking of the client
+			PublicDelayLoadDLLs.Add("GameInput.dll");
+		}
+
+		public GameInputBase(ReadOnlyTargetRules Target) : base(Target)
+		{
+			// Enable truncation warnings in this plugin
+			UnsafeTypeCastWarningLevel = WarningLevel.Error;
+
+			// Uncomment this line to make for easier debugging
+			//OptimizeCode = CodeOptimization.Never;
+
+			PublicDependencyModuleNames.AddRange(
+				new string[]
+				{
+					"Core",
+					"ApplicationCore",
+					"Slate",
+					"Engine",
+					"InputCore",
+					"InputDevice",				
+					"CoreUObject",
+					"DeveloperSettings",
+				}
+			);
+
+			bool bHasGameInputSupport = HasGameInputSupport(Target);
+			// Define this as 0 in the base module to avoid compilation errors when building
+			// without any Game Input support. It is up to the platform-specific submodules to define 
+			PublicDefinitions.Add("GAME_INPUT_SUPPORT=" + (bHasGameInputSupport ? "1" : "0"));
+
+			// Give platforms extensions a chance to add any required dependencies that may be necessary to compile game input
+			if (bHasGameInputSupport)
+			{
+				AddRequiredDeps(Target);
+			}
+		}
+						
+		/**
+		 * Returns the SDK version number based on the installation path.
+		 * This will first attempt to read the version from the grdk.ini file.
+		 * If that fails, then we will try and parse the version from the file path of the installation.
+		 *
+		 * Returns a string of the version number, or "-1" if it couldn't be found. 
+		 */
+		private string GetGameInputVersionString(in string GRDKPath)
+		{
+			string GRDKVersionString = GetGameInputVersionStringFromIni(GRDKPath);
+			if (GRDKVersionString == "-1")
+			{
+				GRDKVersionString = GetGameInputVersionStringFromPath(GRDKPath);
+			}
+
+			return GRDKVersionString;
+		}
+		
+		/**
+		 * Looks for the grdk.ini file in the given GRDK path.
+		 * If the file exists, then attempt to find the version in the ini
+		 *
+		 * Returns a string of the version from the ini if successful, but "-1" if it fails. 
+		 */
+		private string GetGameInputVersionStringFromIni(in string GRDKPath)
+		{
+			if (!string.IsNullOrEmpty(GRDKPath) && Directory.Exists(GRDKPath))
+			{
+				string IniFileName = Path.Combine(GRDKPath, "grdk.ini");
+				FileInfo Info = new FileInfo(IniFileName);
+				if (Info.Exists)
+				{
+					// Parse the file for the _xbld_edition value
+					// it will be in the format of 
+					//  _xbld_edition=<version>
+					// i.e.
+					//  _xbld_edition=230306
+					foreach (string Line in File.ReadAllLines(IniFileName))
+					{
+						if (!Line.Contains("_xbld_edition="))
+						{
+							continue;
+						}
+
+						string[] keyVal = Line.Split('=');
+						if (keyVal.Length >= 2)
+						{
+							return keyVal[1];
+						}
+					}
+				}
+			}
+
+			return "-1";
+		}
+
+		/**
+		 * Returns the 6 digit version number based on the installation path of the SDK.
+		 * -1 if the path is invalid or a valid version couldn't be found.
+		 */
+		private string GetGameInputVersionStringFromPath(in string GRDKPath)
+		{
+			if (GRDKPath is null)
+			{
+				return "-1";
+			}
+
+			// regex for finding a 6 digit number before "GRDK", which should always give us the version number of the GRDK that is installed
+			string Pattern = "(\\d{6})(\\\\GRDK)";
+			Regex Rg = new Regex(Pattern);
+
+			MatchCollection Matched = Rg.Matches(GRDKPath);
+
+			// The result we care about will be in the first match, second group
+			return
+				Matched.Count > 0 && Matched[0].Groups.Count > 1 ?
+					Matched[0].Groups[1].ToString()
+					: "-1";
+		}
+	}
+}
diff --git a/Engine/Plugins/Runtime/GameInput/Source/GameInputBase/Private/GameInputBaseModule.cpp b/Engine/Plugins/Runtime/GameInput/Source/GameInputBase/Private/GameInputBaseModule.cpp
new file mode 100644
--- /dev/null
+++ b/Engine/Plugins/Runtime/GameInput/Source/GameInputBase/Private/GameInputBaseModule.cpp
@@ -0,0 +1,175 @@
+// Copyright Epic Games, Inc. All Rights Reserved.
+
+#include "GameInputBaseModule.h"
+
+#include "CoreGlobals.h"
+#include "GameInputBaseIncludes.h"
+#include "GameInputKeyTypes.h"
+#include "GameInputLogging.h"
+#include "Misc/App.h"
+#include "Modules/ModuleManager.h"
+
+#if GAME_INPUT_SUPPORT
+THIRD_PARTY_INCLUDES_START
+#include "Microsoft/AllowMicrosoftPlatformTypes.h"
+#include <GameInput.h>
+#include "Microsoft/HideMicrosoftPlatformTypes.h"
+THIRD_PARTY_INCLUDES_END
+#include "Microsoft/COMPointer.h"
+#endif
+
+#define LOCTEXT_NAMESPACE "GameInputBaseModule" 
+
+namespace UE::GameInput
+{
+#if GAME_INPUT_SUPPORT
+	// A singleton pointer to the base GameInput interface. 
+	// This provides access to reading the input stream, device callbacks, and more. 
+	static TComPtr<IGameInput> GGameInputInterface;
+#endif	// #if GAME_INPUT_SUPPORT	
+}
+
+FGameInputBaseModule& FGameInputBaseModule::Get()
+{
+	return FModuleManager::LoadModuleChecked<FGameInputBaseModule>(TEXT("GameInputBase"));
+}
+
+bool FGameInputBaseModule::IsAvailable()
+{
+	return FModuleManager::Get().IsModuleLoaded(TEXT("GameInputBase"));
+}
+
+void FGameInputBaseModule::StartupModule()
+{
+	UE_LOG(LogGameInput, Log, TEXT("GameInputBase module startup..."));
+	
+	// We don't care for Game Input if we are running a commandlet, like when we are cooking.
+	if (IsRunningCommandlet())
+	{
+		UE_LOG(LogGameInput, Log, TEXT("GameInputBase module is exiting because IsRunningCommandlet is true. GameInput will not be initalized."));
+		return;
+	}
+
+	// If there is no project name then we don't need game input either. This means we are in the project launcher
+	if (!FApp::HasProjectName())
+	{
+		UE_LOG(LogGameInput, Log, TEXT("GameInputBase module is exiting because there is no project name. GameInput will not be initalized."));
+		return;
+	}
+
+	// Unattended app can't receive any user input, so there is no need to try and create the GameInput interface.
+	if (FApp::IsUnattended())
+	{
+		UE_LOG(LogGameInput, Log, TEXT("GameInputBase module is exiting because it is unattended (FApp::IsUnattended is true) and thus cannot recieve user input. GameInput will not be initalized."));
+		return;
+	}
+
+	// Doesn't make sense to have headless apps create game input
+	if (!FApp::CanEverRender())
+	{
+		UE_LOG(LogGameInput, Log, TEXT("GameInputBase module is exiting because it cannot render anything (FApp::CanEverRender is false). GameInput will not be initalized."));
+		return;
+	}
+
+#if GAME_INPUT_SUPPORT
+
+	// Create the Game Input interface
+	const HRESULT HResult = GameInputCreate(&UE::GameInput::GGameInputInterface);
+	UE_CLOG(FAILED(HResult), LogGameInput, Warning, TEXT("Failed to initialize GameInput: 0x%X"), HResult);
+
+	if (SUCCEEDED(HResult))
+	{
+		UE_LOG(LogGameInput, Log, TEXT("[FGameInputBaseModule::StartupModule] Successfully created the IGameInput interface"));
+		InitializeGameInputKeys();
+	}
+
+#else
+	UE_LOG(LogGameInput, Warning, TEXT("Failed to initalize GameInput! GAME_INPUT_SUPPORT is false!"));
+#endif	// #if GAME_INPUT_SUPPORT
+}
+
+void FGameInputBaseModule::ShutdownModule()
+{
+#if GAME_INPUT_SUPPORT
+	UE::GameInput::GGameInputInterface.Reset();
+#endif // #if GAME_INPUT_SUPPORT
+}
+
+#if GAME_INPUT_SUPPORT
+IGameInput* FGameInputBaseModule::GetGameInput()
+{
+	return UE::GameInput::GGameInputInterface;
+}
+#endif	// GAME_INPUT_SUPPORT
+
+
+void FGameInputBaseModule::InitializeGameInputKeys()
+{
+#if GAME_INPUT_SUPPORT
+
+	static const FName MenuCategory = TEXT("GameInput");
+	EKeys::AddMenuCategoryDisplayInfo(MenuCategory, LOCTEXT("GameInput", "Game Input"), TEXT("GraphEditor.PadEvent_16x"));
+
+	//  
+	// Racing Wheel
+	//
+
+	// Analog types
+	EKeys::AddKey(FKeyDetails(FGameInputKeys::RacingWheel_Brake, LOCTEXT("GameInput_RacingWheel_Brake", "Game Input Racing Wheel Brake"), FKeyDetails::Axis1D, MenuCategory));
+	EKeys::AddKey(FKeyDetails(FGameInputKeys::RacingWheel_Clutch, LOCTEXT("GameInput_RacingWheel_Clutch", "Game Input Racing Wheel Clutch"), FKeyDetails::Axis1D, MenuCategory));
+	EKeys::AddKey(FKeyDetails(FGameInputKeys::RacingWheel_Handbrake, LOCTEXT("GameInput_RacingWheel_Handbrake", "Game Input Racing Wheel Handbrake"), FKeyDetails::Axis1D, MenuCategory));
+	EKeys::AddKey(FKeyDetails(FGameInputKeys::RacingWheel_Throttle, LOCTEXT("GameInput_RacingWheel_Throttle", "Game Input Racing Wheel Throttle"), FKeyDetails::Axis1D, MenuCategory));
+	EKeys::AddKey(FKeyDetails(FGameInputKeys::RacingWheel_Wheel, LOCTEXT("GameInput_RacingWheel_Wheel", "Game Input Racing Wheel"), FKeyDetails::Axis1D, MenuCategory));
+	EKeys::AddKey(FKeyDetails(FGameInputKeys::RacingWheel_PatternShifterGear, LOCTEXT("GameInput_RacingWheel_PatternShifterGear", "Game Input Racing Wheel Pattern Shifter Gear"), FKeyDetails::Axis1D, MenuCategory));
+
+	// Button types
+	EKeys::AddKey(FKeyDetails(FGameInputKeys::RacingWheel_None, LOCTEXT("GameInput_RacingWheel_None", "Game Input Racing Wheel None"), FKeyDetails::GamepadKey, MenuCategory));
+	EKeys::AddKey(FKeyDetails(FGameInputKeys::RacingWheel_Menu, LOCTEXT("GameInput_RacingWheel_Menu", "Game Input Racing Wheel Menu"), FKeyDetails::GamepadKey, MenuCategory));
+	EKeys::AddKey(FKeyDetails(FGameInputKeys::RacingWheel_View, LOCTEXT("GameInput_RacingWheel_View", "Game Input Racing Wheel View"), FKeyDetails::GamepadKey, MenuCategory));
+	EKeys::AddKey(FKeyDetails(FGameInputKeys::RacingWheel_PreviousGear, LOCTEXT("GameInput_RacingWheel_PreviousGear", "Game Input Racing Wheel Previous Gear"), FKeyDetails::GamepadKey, MenuCategory));
+	EKeys::AddKey(FKeyDetails(FGameInputKeys::RacingWheel_NextGear, LOCTEXT("GameInput_RacingWheel_NextGear", "Game Input Racing Wheel Next Gear"), FKeyDetails::GamepadKey, MenuCategory));
+
+	//
+	// Flight Stick
+	//
+
+	// Analog types
+	EKeys::AddKey(FKeyDetails(FGameInputKeys::FlightStick_Roll, LOCTEXT("GameInput_FlightStick_Roll", "Game Input Flight Stick Roll"), FKeyDetails::Axis1D, MenuCategory));
+	EKeys::AddKey(FKeyDetails(FGameInputKeys::FlightStick_Pitch, LOCTEXT("GameInput_FlightStick_Pitch", "Game Input Flight Stick Pitch"), FKeyDetails::Axis1D, MenuCategory));
+	EKeys::AddKey(FKeyDetails(FGameInputKeys::FlightStick_Yaw, LOCTEXT("GameInput_FlightStick_Yaw", "Game Input Flight Stick Yaw"), FKeyDetails::Axis1D, MenuCategory));
+	EKeys::AddKey(FKeyDetails(FGameInputKeys::FlightStick_Throttle, LOCTEXT("GameInput_FlightStick_Throttle", "Game Input Flight Stick Throttle"), FKeyDetails::Axis1D, MenuCategory));
+
+	// Button types
+	EKeys::AddKey(FKeyDetails(FGameInputKeys::FlightStick_None, LOCTEXT("GameInput_FlightStick_None", "Game Input Flight Stick None"), FKeyDetails::GamepadKey, MenuCategory));
+	EKeys::AddKey(FKeyDetails(FGameInputKeys::FlightStick_Menu, LOCTEXT("GameInput_FlightStick_Menu", "Game Input Flight Stick Menu"), FKeyDetails::GamepadKey, MenuCategory));
+	EKeys::AddKey(FKeyDetails(FGameInputKeys::FlightStick_View, LOCTEXT("GameInput_FlightStick_View", "Game Input Flight Stick View"), FKeyDetails::GamepadKey, MenuCategory));
+	EKeys::AddKey(FKeyDetails(FGameInputKeys::FlightStick_FirePrimary, LOCTEXT("GameInput_FlightStick_FirePrimary", "Game Input Flight Stick Fire Primary"), FKeyDetails::GamepadKey, MenuCategory));
+	EKeys::AddKey(FKeyDetails(FGameInputKeys::FlightStick_FireSecondary, LOCTEXT("GameInput_FlightStick_FireSecondary", "Game Input Flight Stick Fire Secondary"), FKeyDetails::GamepadKey, MenuCategory));
+
+	//
+	// Arcade Stick    
+	//
+
+	// Button Types
+	EKeys::AddKey(FKeyDetails(FGameInputKeys::ArcadeStick_None, LOCTEXT("GameInput_ArcadeStick_None", "Game Input Arcade Stick None"), FKeyDetails::GamepadKey, MenuCategory));
+	EKeys::AddKey(FKeyDetails(FGameInputKeys::ArcadeStick_Menu, LOCTEXT("GameInput_ArcadeStick_Menu", "Game Input Arcade Stick Menu"), FKeyDetails::GamepadKey, MenuCategory));
+	EKeys::AddKey(FKeyDetails(FGameInputKeys::ArcadeStick_View, LOCTEXT("GameInput_ArcadeStick_View", "Game Input Arcade Stick View"), FKeyDetails::GamepadKey, MenuCategory));
+	EKeys::AddKey(FKeyDetails(FGameInputKeys::ArcadeStick_Up, LOCTEXT("GameInput_ArcadeStick_Up", "Game Input Arcade Stick Up"), FKeyDetails::GamepadKey, MenuCategory));
+	EKeys::AddKey(FKeyDetails(FGameInputKeys::ArcadeStick_Down, LOCTEXT("GameInput_ArcadeStick_Down", "Game Input Arcade Stick Down"), FKeyDetails::GamepadKey, MenuCategory));
+	EKeys::AddKey(FKeyDetails(FGameInputKeys::ArcadeStick_Left, LOCTEXT("GameInput_ArcadeStick_Left", "Game Input Arcade Stick Left"), FKeyDetails::GamepadKey, MenuCategory));
+	EKeys::AddKey(FKeyDetails(FGameInputKeys::ArcadeStick_Right, LOCTEXT("GameInput_ArcadeStick_Right", "Game Input Arcade Stick Right"), FKeyDetails::GamepadKey, MenuCategory));
+	EKeys::AddKey(FKeyDetails(FGameInputKeys::ArcadeStick_Action1, LOCTEXT("GameInput_ArcadeStick_Action1", "Game Input Arcade Stick Action 1"), FKeyDetails::GamepadKey, MenuCategory));
+	EKeys::AddKey(FKeyDetails(FGameInputKeys::ArcadeStick_Action2, LOCTEXT("GameInput_ArcadeStick_Action2", "Game Input Arcade Stick Action 2"), FKeyDetails::GamepadKey, MenuCategory));
+	EKeys::AddKey(FKeyDetails(FGameInputKeys::ArcadeStick_Action3, LOCTEXT("GameInput_ArcadeStick_Action3", "Game Input Arcade Stick Action 3"), FKeyDetails::GamepadKey, MenuCategory));
+	EKeys::AddKey(FKeyDetails(FGameInputKeys::ArcadeStick_Action4, LOCTEXT("GameInput_ArcadeStick_Action4", "Game Input Arcade Stick Action 4"), FKeyDetails::GamepadKey, MenuCategory));
+	EKeys::AddKey(FKeyDetails(FGameInputKeys::ArcadeStick_Action5, LOCTEXT("GameInput_ArcadeStick_Action5", "Game Input Arcade Stick Action 5"), FKeyDetails::GamepadKey, MenuCategory));
+	EKeys::AddKey(FKeyDetails(FGameInputKeys::ArcadeStick_Action6, LOCTEXT("GameInput_ArcadeStick_Action6", "Game Input Arcade Stick Action 6"), FKeyDetails::GamepadKey, MenuCategory));
+	EKeys::AddKey(FKeyDetails(FGameInputKeys::ArcadeStick_Special1, LOCTEXT("GameInput_ArcadeStick_Special1", "Game Input Arcade Stick Special 1"), FKeyDetails::GamepadKey, MenuCategory));
+	EKeys::AddKey(FKeyDetails(FGameInputKeys::ArcadeStick_Special2, LOCTEXT("GameInput_ArcadeStick_Special2", "Game Input Arcade Stick Special 2"), FKeyDetails::GamepadKey, MenuCategory));
+
+#endif	// GAME_INPUT_SUPPORT
+}
+
+IMPLEMENT_MODULE(FGameInputBaseModule, GameInputBase)
+
+#undef LOCTEXT_NAMESPACE
\ No newline at end of file
diff --git a/Engine/Plugins/Runtime/GameInput/Source/GameInputBase/Private/GameInputDeveloperSettings.cpp b/Engine/Plugins/Runtime/GameInput/Source/GameInputBase/Private/GameInputDeveloperSettings.cpp
new file mode 100644
--- /dev/null
+++ b/Engine/Plugins/Runtime/GameInput/Source/GameInputBase/Private/GameInputDeveloperSettings.cpp
@@ -0,0 +1,198 @@
+// Copyright Epic Games, Inc. All Rights Reserved.
+
+#include "GameInputDeveloperSettings.h"
+#include "Misc/DataDrivenPlatformInfoRegistry.h"
+#include "GameInputLogging.h"
+
+#if GAME_INPUT_SUPPORT
+THIRD_PARTY_INCLUDES_START
+#include "Microsoft/AllowMicrosoftPlatformTypes.h"
+#include <GameInput.h>
+#include "Microsoft/HideMicrosoftPlatformTypes.h"
+THIRD_PARTY_INCLUDES_END
+#endif	// GAME_INPUT_SUPPORT
+
+FGameInputDeviceIdentifier::FGameInputDeviceIdentifier()
+	: VendorId(0)
+	, ProductId(0)
+{
+}
+
+FGameInputDeviceIdentifier::FGameInputDeviceIdentifier(uint16 InVendorId, uint16 InProductId)
+	: VendorId(InVendorId)
+	, ProductId(InProductId)
+{
+}
+
+uint32 GetTypeHash(const FGameInputDeviceIdentifier& InId)
+{
+	uint32 Hash = 0;
+	Hash = HashCombine(Hash, GetTypeHash(InId.VendorId));
+	Hash = HashCombine(Hash, GetTypeHash(InId.ProductId));
+	return Hash;
+}
+
+FString FGameInputDeviceIdentifier::ToString() const
+{
+	return FString::Format(TEXT("VendorId: {0} ProductId {1}"), { VendorId, ProductId });
+}
+
+bool FGameInputDeviceIdentifier::operator==(const FGameInputDeviceIdentifier& Other) const
+{
+	return 
+		VendorId == Other.VendorId && 
+		ProductId == Other.ProductId;
+}
+
+bool FGameInputDeviceIdentifier::operator!=(const FGameInputDeviceIdentifier& Other) const
+{
+	return !FGameInputDeviceIdentifier::operator==(Other);
+}
+
+FGameInputDeviceConfiguration::FGameInputDeviceConfiguration()
+	: DeviceIdentifier(0, 0)
+	, bOverrideHardwareDeviceIdString(false)
+	, OverriddenHardwareDeviceId(TEXT(""))
+	, bProcessControllerButtons(true)
+	, bProcessControllerSwitchState(true)
+	, bProcessControllerAxis(true)
+	, bProcessRawReportData(false)
+	, RawReportReadingId(0)
+{
+}
+
+#if GAME_INPUT_SUPPORT
+const FGameInputDeviceConfiguration* UGameInputDeveloperSettings::FindDeviceConfiguration(const GameInputDeviceInfo* const Info) const
+{
+	if (!Info)
+	{
+		return nullptr;
+	}
+
+	return DeviceConfigurations.FindByPredicate([&Info](const FGameInputDeviceConfiguration& DeviceConfig)
+	{
+		// Look for a matching Vendor ID and Product ID
+		return
+			(DeviceConfig.DeviceIdentifier.VendorId == Info->vendorId &&
+			DeviceConfig.DeviceIdentifier.ProductId == Info->productId);
+	});
+}
+
+const FGameInputDeviceConfiguration* UGameInputDeveloperSettings::FindDeviceConfiguration(const FGameInputDeviceIdentifier& HardwareID) const
+{
+	return DeviceConfigurations.FindByPredicate([&HardwareID](const FGameInputDeviceConfiguration& DeviceConfig)
+	{
+		// Look for a matching Vendor ID and Product ID
+		return
+			(DeviceConfig.DeviceIdentifier.VendorId == HardwareID.VendorId &&
+			DeviceConfig.DeviceIdentifier.ProductId == HardwareID.ProductId);
+	});
+}
+
+#endif	// GAME_INPUT_SUPPORT
+
+UGameInputPlatformSettings* UGameInputPlatformSettings::Get()
+{
+	return UPlatformSettingsManager::Get().GetSettingsForPlatform<UGameInputPlatformSettings>();
+}
+
+void UGameInputPlatformSettings::InitializePlatformDefaults()
+{
+	// For platforms who are defaulted to KBM (Desktops) we want to disable KBM processing by default because
+	// it will already be handled by the SlateApplication and platform message pump
+	const FName PlatformName = GetPlatformIniName();
+	const FDataDrivenPlatformInfo& PlatformInfo = FDataDrivenPlatformInfoRegistry::GetPlatformInfo(PlatformName);
+	if (PlatformInfo.DefaultInputType == "MouseAndKeyboard")
+	{
+		bProcessKeyboard = false;
+		bProcessMouse = false;
+
+		// We also want to disable gamepad processing by default here, because it is likely that your PC platform
+		// will be using the existing Xinput implementation. If one day we can turn off XInput by default,
+		// then we can change this to be default on
+		bProcessGamepad = false;
+		bProcessController = false;
+		bProcessRawInput = false;
+	}
+
+	// Everything else can just use the default values already set until we want to change something else
+}
+
+UGameInputDeveloperSettings::UGameInputDeveloperSettings(const FObjectInitializer& Initializer)
+	: Super(Initializer)
+	, bDoNotProcessDuplicateCapabilitiesForSingleUser(true)
+{
+	PlatformSpecificSettings.Initialize(UGameInputPlatformSettings::StaticClass());
+}
+
+#if WITH_EDITOR
+
+namespace UE::GameInput
+{
+	/** A set of FName's that, if change in PostEditChangeProperty, would require Game Input to update what device status updates it is listening for. */
+	static const TSet<FName>& GetPropertyNamesThatRequireGameInputUpdate()
+	{
+		static const TSet<FName> PropertyNamesThatRequireGameInputUpdate
+		{
+			GET_MEMBER_NAME_CHECKED(UGameInputPlatformSettings, bProcessController),
+		};
+		return PropertyNamesThatRequireGameInputUpdate;
+	}
+}
+
+const TArray<uint32>& UGameInputDeveloperSettings::GetControllerButtonMappingDataKeyOptions()
+{
+	// The controller button bitmask is generated by bit shifting the current index
+	// so just display all the possible options here
+	static const TArray<uint32> Options =
+	{
+		1 << 0,
+		1 << 1,
+		1 << 2,
+		1 << 3,
+		1 << 4,
+		1 << 5,
+		1 << 6,
+		1 << 7,
+		1 << 8,
+		1 << 9,
+		1 << 10,
+		1 << 11,
+		1 << 12,
+		1 << 13,
+		1 << 14,
+		1 << 15,
+		1 << 16,
+		1 << 17,
+		1 << 18,
+		1 << 19,
+		1 << 21,
+		1 << 22,
+		1 << 23,
+		1 << 24,
+		1 << 25,
+		1 << 26,
+		1 << 27,
+		1 << 28,
+		1 << 29,
+		1 << 30,
+	};
+
+	return Options;
+}
+
+void UGameInputDeveloperSettings::PostEditChangeProperty(FPropertyChangedEvent& PropertyChangedEvent)
+{
+	Super::PostEditChangeProperty(PropertyChangedEvent);
+
+	const FName PropertyName = (PropertyChangedEvent.Property != nullptr) ? PropertyChangedEvent.Property->GetFName() : NAME_None;
+
+	if (UE::GameInput::GetPropertyNamesThatRequireGameInputUpdate().Contains(PropertyName))
+	{
+		if (OnInputSettingChanged.IsBound())
+		{
+			OnInputSettingChanged.Execute(PropertyName);
+		}
+	}
+}
+#endif	// WITH_EDITOR
\ No newline at end of file
diff --git a/Engine/Plugins/Runtime/GameInput/Source/GameInputBase/Private/GameInputDeviceContainer.cpp b/Engine/Plugins/Runtime/GameInput/Source/GameInputBase/Private/GameInputDeviceContainer.cpp
new file mode 100644
--- /dev/null
+++ b/Engine/Plugins/Runtime/GameInput/Source/GameInputBase/Private/GameInputDeviceContainer.cpp
@@ -0,0 +1,376 @@
+// Copyright Epic Games, Inc. All Rights Reserved.
+
+#include "GameInputDeviceContainer.h"
+#include "Framework/Application/SlateApplication.h"	// for GetPlatformCursor
+#include "GenericPlatform/ICursor.h"
+#include "GameInputUtils.h"
+#include "GameInputLogging.h"
+#include "GameInputDeveloperSettings.h"
+#include "HAL/IConsoleManager.h"
+
+#if GAME_INPUT_SUPPORT
+
+namespace UE::GameInput
+{
+	static TAutoConsoleVariable<bool> CVarGameInputForceDeviceProcesorsOncePerFrame
+	(
+		TEXT("gameinput.bForceDeviceProcesorsOncePerFrame"),
+		false,
+		TEXT("If true, then FGameInputDeviceContainer::ProcessInput will ensure that device processors are only called once per ProcessInput call (i.e. once per frame)"),
+		ECVF_Default
+	);
+};
+
+FGameInputDeviceContainer::FGameInputDeviceContainer(
+	const TSharedRef<FGenericApplicationMessageHandler>& InMessageHandler,
+	IGameInputDevice* InDevice,
+	GameInputKind InAllowedGameInputKinds,
+	FPlatformUserId InUserId, 
+	FInputDeviceId InDeviceId)
+	: MessageHandler(InMessageHandler)
+	, Device(InDevice)
+	, AllowedGameInputKinds(InAllowedGameInputKinds)
+	, UserId(InUserId)
+	, AssignedDeviceId(InDeviceId)
+	, IgnoreReadingTimestamp(0)
+{
+	SetGameInputDevice(InDevice);
+}
+
+void FGameInputDeviceContainer::SetMessageHandler(const TSharedRef<FGenericApplicationMessageHandler>& InMessageHandler)
+{
+	MessageHandler = InMessageHandler;
+}
+
+void FGameInputDeviceContainer::InitalizeDeviceProcessors()
+{
+	if (!Device)
+	{
+		return;
+	}
+
+	// Based on the GameInputKind of this device, create any processors that are supported
+	const GameInputDeviceInfo* Info = Device->GetDeviceInfo();
+
+	const GameInputKind DeviceKind = Info->supportedInput;
+
+	UE_LOG(LogGameInput, Log, TEXT("InitalizeDeviceProcessors for device kind %s with an allowed kind of %s"), *UE::GameInput::LexToString(DeviceKind), *UE::GameInput::LexToString(AllowedGameInputKinds));
+
+	InitalizeDeviceProcessors_Impl();
+}
+
+void FGameInputDeviceContainer::InitalizeDeviceProcessors_Impl()
+{
+	if (!Device)
+	{
+		return;
+	}
+
+	// Based on the GameInputKind of this device, create any processors that are supported
+	const GameInputDeviceInfo* Info = Device->GetDeviceInfo();
+
+	const GameInputKind DeviceKind = Info->supportedInput;
+
+	if (DeviceKind & AllowedGameInputKinds & GameInputKindGamepad)
+	{
+		Processors.Emplace(MakeShared<FGameInputGamepadDeviceProcessor>());
+	}
+	
+	if (DeviceKind & AllowedGameInputKinds & GameInputKindKeyboard)
+	{
+		Processors.Emplace(MakeShared<FGameInputKeyboardDeviceProcessor>());
+	}
+	
+	// TODO: This was throwing a linker error on some platforms, investigate.
+	//if (DeviceKind & AllowedGameInputKinds & GameInputKindMouse)
+	//{
+	//	// TODO: Make high precision mouse a config setting
+	//	Processors.Emplace(MakeShared<FGameInputMouseDeviceProcessor>(FSlateApplication::Get().GetPlatformCursor()));
+	//}
+
+	if (DeviceKind & AllowedGameInputKinds & GameInputKindTouch)
+	{
+		Processors.Emplace(MakeShared<FGameInputTouchDeviceProcessor>());
+	}
+
+	// Only allow for external controller devices have have specified configs if that is enabled
+	const UGameInputDeveloperSettings* Settings = GetDefault<UGameInputDeveloperSettings>();
+	const UGameInputPlatformSettings* PlatformSettings = UGameInputPlatformSettings::Get();
+	
+	// "Controller" types
+	if (DeviceKind & AllowedGameInputKinds & GameInputKindController)
+	{
+		const bool bIsDeviceAllowed =
+			PlatformSettings->bSpecialDevicesRequireExplicitDeviceConfiguration ?
+			(Settings->FindDeviceConfiguration(Info) != nullptr) :
+			true;
+
+		if (bIsDeviceAllowed)
+		{
+			Processors.Emplace(MakeShared<FGameInputControllerDeviceProcessor>());
+		}
+		else
+		{
+			UE_LOG(
+				LogGameInput,
+				Warning,
+				TEXT("A game input controller device (%s) was connected but will not be processed because it does not have an explict device configuration"),
+				*UE::GameInput::LexToString(Device));
+		}
+	}
+
+	// Raw Input
+	if (DeviceKind & AllowedGameInputKinds & GameInputKindRawDeviceReport)
+	{
+		const bool bIsDeviceAllowed =
+			PlatformSettings->bSpecialDevicesRequireExplicitDeviceConfiguration ?
+			(Settings->FindDeviceConfiguration(Info) != nullptr) :
+			true;
+
+		if (bIsDeviceAllowed)
+		{
+			Processors.Emplace(MakeShared<FGameInputRawDeviceProcessor>());
+		}
+		else
+		{
+			UE_LOG(
+				LogGameInput,
+				Warning,
+				TEXT("A raw game input device (%s) was connected but will not be processed because it does not have an explict device configuration"),
+				*UE::GameInput::LexToString(Device));
+		}
+	}
+
+	// TODO: Implement these kinds of input! 
+	
+	//if (DeviceKind & AllowedGameInputKinds &  GameInputKindController)
+	//{
+
+	//}
+
+	//if (DeviceKind & AllowedGameInputKinds & GameInputKindRacingWheel)
+	//{
+
+	//}
+
+	//if (DeviceKind & AllowedGameInputKinds & GameInputKindFlightStick)
+	//{
+
+	//}
+
+	//if (DeviceKind & AllowedGameInputKinds & GameInputKindArcadeStick)
+	//{
+
+	//}
+
+	//if (DeviceKind &AllowedGameInputKinds &  GameInputKindRawDeviceReport)
+	//{
+
+	//}
+}
+
+const GameInputKind FGameInputDeviceContainer::ProcessInput(IGameInput* GameInput, const GameInputKind CurrentSupportedKind, const GameInputKind ProcessedKindsForPlatformUserThisFrame)
+{
+	GameInputKind OutProcessedInputKinds = GameInputKindUnknown;
+
+	const bool bDoNotProcessDuplicateCapabilitiesForSingleUser = GetDefault<UGameInputDeveloperSettings>()->bDoNotProcessDuplicateCapabilitiesForSingleUser;
+
+	// keep reading the input snapshots for this device
+	if (!LastReading.IsValid())
+	{				
+		GameInput->GetCurrentReading(CurrentSupportedKind, Device, &LastReading);
+	}
+
+	// Keep track of the number of input device processors we have used this frame.
+	// We only want to process each one once...
+	static TSet<TSharedPtr<IGameInputDeviceProcessor>> UsedProcessorsThisFrame;
+	UsedProcessorsThisFrame.Reset();
+
+	// Set of processors that we should skip because the platform user has already processed one
+	// of their type this frame already.
+	static TSet<TSharedPtr<IGameInputDeviceProcessor>> ProcessorsToSkipThisFrame;
+	ProcessorsToSkipThisFrame.Reset();
+
+	int32 NumReadingsProcessed = 0;
+
+	while (LastReading.IsValid())
+	{
+		// The current game input reading. This can be null if there is nothing else in the input stack
+		TComPtr<IGameInputReading> Reading;
+		HRESULT hr = GameInput->GetNextReading(LastReading.Get(), CurrentSupportedKind, Device, &Reading);
+		
+		// On the last reading of the frame, there will not be a "next" reading.
+		const bool bIsLastReadingOfFrame = (hr == GAMEINPUT_E_READING_NOT_FOUND);
+
+		if (FAILED(hr))
+		{
+			if (hr != GAMEINPUT_E_READING_NOT_FOUND)
+			{
+				// unexpected error - start from scratch next frame
+				LastReading.Reset();
+				break;
+			}
+		}
+
+		// ignore this input if we're suppressing input for a while
+		if (Reading && Reading->GetTimestamp() < IgnoreReadingTimestamp)
+		{
+			LastReading = Reading;
+			return OutProcessedInputKinds;
+		}
+
+		const GameInputKind CurrentReadingKind = Reading ? Reading->GetInputKind() : LastReading->GetInputKind();
+
+		// Pass along this reading to our Input Processors, who will actually do the work of sending messages
+		// and events to the message handler
+		IGameInputDeviceProcessor::FGameInputEventParams Params = {};
+		Params.Reading = Reading;
+		Params.PreviousReading = LastReading;
+		Params.Device = Device;
+		Params.MessageHandler = MessageHandler;
+		Params.PlatformUserId = UserId;
+		Params.InputDeviceId = AssignedDeviceId;
+
+		const bool bEnforceOncePerFrameProcessors = UE::GameInput::CVarGameInputForceDeviceProcesorsOncePerFrame.GetValueOnAnyThread();
+
+		for (TSharedPtr<IGameInputDeviceProcessor>& Processor : Processors)
+		{
+			// Skip any processors that have already been used this frame
+			if (bEnforceOncePerFrameProcessors && UsedProcessorsThisFrame.Contains(Processor))
+			{
+				continue;
+			}
+
+			const GameInputKind ProcessorKind = Processor->GetSupportedReadingKind();
+
+			// The first time we process this, check for if the user has already handled a processor of this kind during this frame.
+			// If they have, then we should skip processing for the rest of this frame.
+			if (NumReadingsProcessed == 0 && bDoNotProcessDuplicateCapabilitiesForSingleUser && (ProcessedKindsForPlatformUserThisFrame & ProcessorKind) != 0)
+			{
+				ProcessorsToSkipThisFrame.Add(Processor);
+			}
+
+			if (ProcessorsToSkipThisFrame.Contains(Processor))
+			{
+				continue;
+			}
+
+			// Only try and process this processor if the GameInput reading has readings for it
+			if ((ProcessorKind & CurrentReadingKind) != 0)
+			{
+				bool bHasInputThisFrame = false;
+				if (bIsLastReadingOfFrame)
+				{
+					// On the last reading input for the frame, the current reading pointer should be null
+					ensure(!Params.Reading && Params.PreviousReading);
+					bHasInputThisFrame |= Processor->PostProcessInput(Params);
+				}
+				else
+				{
+					// When processing normal input we should always have a current
+					// and previous reading
+					ensure(Params.Reading && Params.PreviousReading);
+					bHasInputThisFrame |= Processor->ProcessInput(Params);
+					++NumReadingsProcessed;
+				}
+
+				UsedProcessorsThisFrame.Add(Processor);
+
+				// Keep track of what input processors have sent events
+				if (bHasInputThisFrame)
+				{
+					OutProcessedInputKinds |= ProcessorKind;
+				}				
+			}
+		}
+
+		// if this was the last reading of the frame, then Reading is going to be null
+		// we dont need to track anything here because it happened on the previous iteration
+		if (bIsLastReadingOfFrame)
+		{
+			break;
+		}
+
+		// Remember this reading so that we can diff against it next frame if we need to
+		LastReading = Reading;
+
+		// Keep track of the timestamp of this reading so that we can later determine the most recently used device
+		LastReadingTimestamp = LastReading->GetTimestamp();
+	}
+
+	UE_LOG(LogGameInput, VeryVerbose, TEXT("Processed '%d' GameInput readings off the input stack for device: %s"), NumReadingsProcessed, *UE::GameInput::LexToString(Device));
+
+	return OutProcessedInputKinds;
+}
+
+void FGameInputDeviceContainer::ClearInputState(IGameInput* GameInput)
+{
+	IGameInputDeviceProcessor::FGameInputEventParams Params = {};
+	Params.Reading = nullptr;	// We have no reading when we clear input
+	Params.PreviousReading = nullptr;
+	Params.Device = Device;
+	Params.MessageHandler = MessageHandler;
+	Params.PlatformUserId = UserId;
+	Params.InputDeviceId = AssignedDeviceId;
+
+	for (TSharedPtr<IGameInputDeviceProcessor>& Processor : Processors)
+	{
+		Processor->ClearState(Params);
+	}
+
+	LastReading.Reset();
+}
+
+IGameInputDevice* FGameInputDeviceContainer::GetGameInputDevice() const
+{
+	return Device;
+}
+
+void FGameInputDeviceContainer::SetGameInputDevice(IGameInputDevice* InDevice)
+{
+	Device = InDevice;
+	if (Device)
+	{
+		if (const GameInputDeviceInfo* Info = InDevice->GetDeviceInfo())
+		{
+			LocalDeviceId = Info->deviceId;
+		}
+		else
+		{
+			// We shouldn't be able to get here, but handle it just in case.
+			ensure(false);
+			LocalDeviceId = {};
+		}
+	}
+	else
+	{
+		LocalDeviceId = {};
+	}
+}
+
+APP_LOCAL_DEVICE_ID FGameInputDeviceContainer::GetGameInputDeviceId() const
+{
+	return LocalDeviceId;
+}
+
+void FGameInputDeviceContainer::SetPlatformUserId(const FPlatformUserId InUserId)
+{
+	UserId = InUserId;
+}
+
+FPlatformUserId FGameInputDeviceContainer::GetPlatformUserId() const
+{
+	return UserId;
+}
+
+void FGameInputDeviceContainer::SetInputDeviceId(const FInputDeviceId InDeviceId)
+{
+	AssignedDeviceId = InDeviceId;
+}
+
+FInputDeviceId FGameInputDeviceContainer::GetDeviceId() const
+{
+	return AssignedDeviceId;
+}
+
+#endif	// GAME_INPUT_SUPPORT
diff --git a/Engine/Plugins/Runtime/GameInput/Source/GameInputBase/Private/GameInputDeviceProcessor.cpp b/Engine/Plugins/Runtime/GameInput/Source/GameInputBase/Private/GameInputDeviceProcessor.cpp
new file mode 100644
--- /dev/null
+++ b/Engine/Plugins/Runtime/GameInput/Source/GameInputBase/Private/GameInputDeviceProcessor.cpp
@@ -0,0 +1,1758 @@
+// Copyright Epic Games, Inc. All Rights Reserved.
+
+#include "GameInputDeviceProcessor.h"
+#include "GameInputUtils.h"
+#include "GameInputLogging.h"
+#include "HAL/PlatformTime.h"
+#include "GameInputKeyTypes.h"
+
+
+#include "GenericPlatform/GenericPlatformInputDeviceMapper.h"
+#include "Framework/Application/SlateApplication.h"
+#include "HAL/IConsoleManager.h"
+#include "GameInputDeveloperSettings.h"
+
+#if GAME_INPUT_SUPPORT
+
+namespace UE::GameInput
+{
+	static FName InputClassName = FName("GameInput");
+};
+
+////////////////////////////////////////////////////////
+// IGameInputDeviceProcessor
+
+const GameInputDeviceInfo* IGameInputDeviceProcessor::FGameInputEventParams::GetDeviceInfo() const
+{
+	if (Device)
+	{
+		return Device->GetDeviceInfo();
+	}
+	return nullptr;
+}
+
+bool IGameInputDeviceProcessor::PostProcessInput(const FGameInputEventParams& Params)
+{
+	// Nothing needs to be done by default here in PostProcessInput.
+	return false;
+}
+
+const FString& IGameInputDeviceProcessor::GetHardwareDeviceIdentifierName(const IGameInputDeviceProcessor::FGameInputEventParams& Params) const
+{
+	static FString ID_Virtual = TEXT("Virtual");
+	static FString ID_Aggregate = TEXT("Aggregate");
+	static FString ID_XboxOne = TEXT("XboxOne");
+	static FString ID_Xbox360 = TEXT("Xbox360");
+	static FString ID_Hid = TEXT("Hid");
+	static FString ID_I8042 = TEXT("I8042");
+
+	if (Params.Device)
+	{
+		if (const GameInputDeviceInfo* Info = Params.GetDeviceInfo())
+		{
+			// Check for any device specific overrides that may be there for custom devices
+			if (const UGameInputDeveloperSettings* Settings = GetDefault<UGameInputDeveloperSettings>())
+			{
+				if (const FGameInputDeviceConfiguration* Config = Settings->FindDeviceConfiguration(Info))
+				{
+					if (Config->bOverrideHardwareDeviceIdString)
+					{
+						return Config->OverriddenHardwareDeviceId;
+					}					
+				}
+			}
+
+			switch (Info->deviceFamily)
+			{
+			case GameInputFamilyVirtual:
+				return ID_Virtual;
+			case GameInputFamilyAggregate:
+				return ID_Aggregate;
+			case GameInputFamilyXbox360:
+				return ID_Xbox360;
+			case GameInputFamilyHid:
+				return ID_Hid;
+			case GameInputFamilyI8042:
+				return ID_I8042;
+			case GameInputFamilyXboxOne:
+			default:
+				return ID_XboxOne;
+			}
+		}
+	}
+
+	// If for some reason we are given a null device, default to just "Xbox One"
+	return ID_XboxOne;
+}
+
+void IGameInputDeviceProcessor::OnControllerAnalog(const FGameInputEventParams& Params, const FName& GamePadKey, float NewAxisValueNormalized, float OldAxisValueNormalized, float DeadZone)
+{
+	if (OldAxisValueNormalized != NewAxisValueNormalized || FMath::Abs(NewAxisValueNormalized) > DeadZone)
+	{		
+		UE_LOG(LogGameInput, VeryVerbose, TEXT("Device %s (PlatformUserId = %d, InputDeviceId = %d) - Analog %s : %.3f"),
+			*UE::GameInput::LexToString(Params.Device),
+			Params.PlatformUserId.GetInternalId(),
+			Params.InputDeviceId.GetId(),
+			*GamePadKey.ToString(),
+			NewAxisValueNormalized);
+
+		// We should only tell slate about this message if the platform user and input device are valid because it will attempt
+		// to create a new slate user based on the index if it doesn't already exist
+		if (Params.PlatformUserId.IsValid() && Params.InputDeviceId.IsValid())
+		{
+			FInputDeviceScope InputScope(nullptr, UE::GameInput::InputClassName, IPlatformInputDeviceMapper::Get().GetUserIndexForPlatformUser(Params.PlatformUserId), GetHardwareDeviceIdentifierName(Params));
+			Params.MessageHandler->OnControllerAnalog(GamePadKey, Params.PlatformUserId, Params.InputDeviceId, NewAxisValueNormalized);
+		}		
+	}
+}
+
+void IGameInputDeviceProcessor::EvaluateButtonStates(
+	const FGameInputEventParams& Params,
+	const uint32 CurrentButtonHeldMask,
+	uint32& PreviousButtonMask,
+	double* RepeatTime,
+	const TMap<uint32, FGamepadKeyNames::Type>& UnrealButtonNameMap,
+	const uint32 SupportedButtonCount /*= MaxSupportedButtons*/)
+{
+	// handle button change events
+	const uint32 ActionMask = (PreviousButtonMask ^ CurrentButtonHeldMask);
+	const uint32 RepeatMask = (PreviousButtonMask & CurrentButtonHeldMask);
+	uint32 BitMask = 1;
+
+	if (!RepeatTime)
+	{
+		UE_LOG(LogGameInput, Error, TEXT("[IGameInputDeviceProcessor::EvaluateButtonStates] Invalid RepeatTime array given to evaluate button states!"));
+		return;
+	}
+
+	// If the given button mask and repeat mask are both zero, then no buttons have been pressed or had a state change. No need to iterate the bitmask
+	if (ActionMask == 0 && RepeatMask == 0)
+	{
+		return;
+	}
+
+	// We can't tell slate about input messages from an invalid platform user
+	if (!Params.PlatformUserId.IsValid())
+	{
+		UE_LOG(LogGameInput, Verbose, TEXT("[IGameInputDeviceProcessor::EvaluateButtonStates] Attempting to evaluate button states with an invalid platform user id of '%d'. The button messages will not be sent."), Params.PlatformUserId.GetInternalId());
+		return;
+	}
+
+	const double CurrentTime = FPlatformTime::Seconds();
+
+	IPlatformInputDeviceMapper& DeviceMapper = IPlatformInputDeviceMapper::Get();
+
+	for (uint32 n = 0; n < SupportedButtonCount; ++n)
+	{
+		FGamepadKeyNames::Type ButtonKey;
+		if (UE::GameInput::GameInputButtonToUnrealName(UnrealButtonNameMap, BitMask, ButtonKey))
+		{
+			// Check for button state change
+			if (0 != (ActionMask & BitMask))
+			{
+				FInputDeviceScope InputScope(nullptr, UE::GameInput::InputClassName, DeviceMapper.GetUserIndexForPlatformUser(Params.PlatformUserId), GetHardwareDeviceIdentifierName(Params));
+
+				if (0 != (CurrentButtonHeldMask & BitMask))
+				{
+					UE_LOG(LogGameInput, Verbose, TEXT("[FGameInputDevice::EvaluateButtonStates] (PlatformUserId = %d, InputDeviceId = %d) - Button '%s' Pressed"), Params.PlatformUserId.GetInternalId(), Params.InputDeviceId.GetId(), *ButtonKey.ToString());
+					Params.MessageHandler->OnControllerButtonPressed(ButtonKey, Params.PlatformUserId, Params.InputDeviceId, false);
+					RepeatTime[n] = CurrentTime + UE::GameInput::InitialRepeatDelay;
+				}
+				else
+				{
+					UE_LOG(LogGameInput, Verbose, TEXT("[FGameInputDevice::EvaluateButtonStates] (PlatformUserId = %d, InputDeviceId = %d) - Button '%s' Released"), Params.PlatformUserId.GetInternalId(), Params.InputDeviceId.GetId(), *ButtonKey.ToString());
+					Params.MessageHandler->OnControllerButtonReleased(ButtonKey, Params.PlatformUserId, Params.InputDeviceId, false);
+					RepeatTime[n] = 0.0;
+				}
+			}
+
+			// Check for repeat key
+			if (0 != (RepeatMask & BitMask) && CurrentTime > RepeatTime[n])
+			{
+				FInputDeviceScope InputScope(nullptr, UE::GameInput::InputClassName, DeviceMapper.GetUserIndexForPlatformUser(Params.PlatformUserId), GetHardwareDeviceIdentifierName(Params));
+
+				RepeatTime[n] = CurrentTime + UE::GameInput::SubsequentRepeatDelay;
+				UE_LOG(LogGameInput, Verbose, TEXT("[FGameInputDevice::EvaluateButtonStates] (PlatformUserId = %d, InputDeviceId = %d) - Button '%s' Repeat Pressed"), Params.PlatformUserId.GetInternalId(), Params.InputDeviceId.GetId(), *ButtonKey.ToString());
+				Params.MessageHandler->OnControllerButtonPressed(ButtonKey, Params.PlatformUserId, Params.InputDeviceId, true);
+			}
+		}
+
+		// Move on to the next bit!
+		BitMask <<= 1;
+	}
+
+	PreviousButtonMask = CurrentButtonHeldMask;
+}
+
+void IGameInputDeviceProcessor::EvaluateSwitchState(
+	const FGameInputEventParams& Params,
+	GameInputSwitchPosition CurrentPosition,
+	GameInputSwitchPosition& PreviousPosition,
+	TArray<double>& RepeatTimes)
+{
+	if (!ensureMsgf(RepeatTimes.IsValidIndex(static_cast<uint32>(GameInputSwitchLeft)), TEXT("RepeatTimes array needs to be the same size as the number of switch positions!")))
+	{
+		return;
+	}
+	GameInputSwitchPosition PrevCopy = PreviousPosition;
+
+	// If the current and previous switch states are both the center, then nothing has happened.
+	if (CurrentPosition == GameInputSwitchCenter && PreviousPosition == GameInputSwitchCenter)
+	{
+		return;
+	}
+
+	// We can't send any slate input events to slate if the platform user is invalid
+	if (!Params.PlatformUserId.IsValid())
+	{
+		UE_LOG(LogGameInput, Verbose, TEXT("[IGameInputDeviceProcessor::EvaluateSwitchState] Attempting to evaluate button states with an invalid platform user id of '%d'. The input messages will not be sent."), Params.PlatformUserId.GetInternalId());
+		return;
+	}
+
+	const double CurrentTime = FPlatformTime::Seconds();
+
+	// If the current and previous are not the same, then release the previous
+	// and send the pressed event for the current, as this is the first press
+
+	if (CurrentPosition != PreviousPosition)
+	{
+		// Release the previous keys
+		if (const TArray<FGamepadKeyNames::Type>* PrevKeyArray = UE::GameInput::SwitchPositionToUnrealName(PreviousPosition))
+		{
+			for (const FGamepadKeyNames::Type KeyName : *PrevKeyArray)
+			{
+				UE_LOG(LogGameInput, Verbose, TEXT("[EvaluateSwitchState] (PlatformUserId = %d, InputDeviceId = %d) - Switch '%s' Released"), Params.PlatformUserId.GetInternalId(), Params.InputDeviceId.GetId(), *KeyName.ToString());
+				Params.MessageHandler->OnControllerButtonReleased(KeyName, Params.PlatformUserId, Params.InputDeviceId, false);
+			}
+
+			RepeatTimes[static_cast<int32>(PreviousPosition)] = 0.0;
+		}
+
+		// "Press" the current key
+		if (const TArray<FGamepadKeyNames::Type>* CurrentKeyArray = UE::GameInput::SwitchPositionToUnrealName(CurrentPosition))
+		{
+			for (const FGamepadKeyNames::Type KeyName : *CurrentKeyArray)
+			{
+				UE_LOG(LogGameInput, Verbose, TEXT("[EvaluateSwitchState] (PlatformUserId = %d, InputDeviceId = %d) - Switch '%s' Pressed"), Params.PlatformUserId.GetInternalId(), Params.InputDeviceId.GetId(), *KeyName.ToString());
+				Params.MessageHandler->OnControllerButtonPressed(KeyName, Params.PlatformUserId, Params.InputDeviceId, true);
+			}
+
+			RepeatTimes[static_cast<int32>(CurrentPosition)] = CurrentTime + UE::GameInput::InitialRepeatDelay;
+		}
+	}
+	// Otherwise, the states are the same. Check if we can repeat them
+	else if (CurrentTime > RepeatTimes[static_cast<int32>(CurrentPosition)])
+	{
+		if (const TArray<FGamepadKeyNames::Type>* CurrentKeyArray = UE::GameInput::SwitchPositionToUnrealName(CurrentPosition))
+		{
+			for (const FGamepadKeyNames::Type KeyName : *CurrentKeyArray)
+			{
+				UE_LOG(LogGameInput, Verbose, TEXT("[EvaluateSwitchState] (PlatformUserId = %d, InputDeviceId = %d) - Switch '%s' Repeat Pressed"), Params.PlatformUserId.GetInternalId(), Params.InputDeviceId.GetId(), *KeyName.ToString());
+				Params.MessageHandler->OnControllerButtonPressed(KeyName, Params.PlatformUserId, Params.InputDeviceId, true);
+			}
+		}
+
+		// Check for repeat key
+		RepeatTimes[static_cast<int32>(CurrentPosition)] = CurrentTime + UE::GameInput::SubsequentRepeatDelay;
+	}
+
+	// Keep track of the previous position
+	PreviousPosition = CurrentPosition;
+
+}
+
+////////////////////////////////////////////////////////
+// FGameInputGamepadDeviceProcessor
+
+namespace UE::GameInput
+{
+	inline bool IsEmptyGamepadReading(const GameInputGamepadState& GamepadState)
+	{
+		return GamepadState.buttons == 0 &&
+			GamepadState.leftTrigger <= GamepadTriggerDeadzone && GamepadState.rightTrigger <= GamepadTriggerDeadzone &&
+			FMath::Abs(GamepadState.leftThumbstickX) <= GamepadLeftStickDeadzone && FMath::Abs(GamepadState.leftThumbstickY) <= GamepadLeftStickDeadzone &&
+			FMath::Abs(GamepadState.rightThumbstickX) <= GamepadRightStickDeadzone && FMath::Abs(GamepadState.rightThumbstickY) <= GamepadRightStickDeadzone;
+	};
+
+	inline bool HasDifferentAnalogInput(const GameInputGamepadState& CurrentGamepadState, const GameInputGamepadState& PreviousGamepadState)
+	{
+		return CurrentGamepadState.leftTrigger != PreviousGamepadState.leftTrigger || CurrentGamepadState.rightTrigger != PreviousGamepadState.rightTrigger ||
+			CurrentGamepadState.leftThumbstickX != PreviousGamepadState.leftThumbstickX || CurrentGamepadState.leftThumbstickY != PreviousGamepadState.leftThumbstickY ||
+			CurrentGamepadState.rightThumbstickX != PreviousGamepadState.rightThumbstickX || CurrentGamepadState.rightThumbstickY != PreviousGamepadState.rightThumbstickY;
+	};
+
+	/**
+	 * Returns a map of GameInputGamepadButtons to their associated Unreal Engine FKey names.
+	 */
+	static const TMap<uint32, FGamepadKeyNames::Type>& GetGamepadButtonMap()
+	{
+		static const TMap<uint32, FGamepadKeyNames::Type> GamepadButtonMap
+		{
+			{ static_cast<uint32>(GameInputGamepadButtons::GameInputGamepadA), FGamepadKeyNames::FaceButtonBottom },
+			{ static_cast<uint32>(GameInputGamepadButtons::GameInputGamepadB), FGamepadKeyNames::FaceButtonRight },
+			{ static_cast<uint32>(GameInputGamepadButtons::GameInputGamepadX), FGamepadKeyNames::FaceButtonLeft },
+			{ static_cast<uint32>(GameInputGamepadButtons::GameInputGamepadY), FGamepadKeyNames::FaceButtonTop },
+			{ static_cast<uint32>(GameInputGamepadButtons::GameInputGamepadLeftShoulder), FGamepadKeyNames::LeftShoulder },
+			{ static_cast<uint32>(GameInputGamepadButtons::GameInputGamepadRightShoulder), FGamepadKeyNames::RightShoulder },
+			{ static_cast<uint32>(GameInputGamepadButtons::GameInputGamepadMenu), FGamepadKeyNames::SpecialRight },
+			{ static_cast<uint32>(GameInputGamepadButtons::GameInputGamepadView), FGamepadKeyNames::SpecialLeft },
+			{ static_cast<uint32>(GameInputGamepadButtons::GameInputGamepadDPadUp), FGamepadKeyNames::DPadUp },
+			{ static_cast<uint32>(GameInputGamepadButtons::GameInputGamepadDPadDown), FGamepadKeyNames::DPadDown },
+			{ static_cast<uint32>(GameInputGamepadButtons::GameInputGamepadDPadLeft), FGamepadKeyNames::DPadLeft },
+			{ static_cast<uint32>(GameInputGamepadButtons::GameInputGamepadDPadRight), FGamepadKeyNames::DPadRight },
+			{ static_cast<uint32>(GameInputGamepadButtons::GameInputGamepadLeftThumbstick), FGamepadKeyNames::LeftThumb },
+			{ static_cast<uint32>(GameInputGamepadButtons::GameInputGamepadRightThumbstick), FGamepadKeyNames::RightThumb },
+
+			{ static_cast<uint32>(EGameInputAuxButton::GamepadButtonAux_LeftTrigger), FGamepadKeyNames::LeftTriggerThreshold },
+			{ static_cast<uint32>(EGameInputAuxButton::GamepadButtonAux_RightTrigger), FGamepadKeyNames::RightTriggerThreshold },
+			{ static_cast<uint32>(EGameInputAuxButton::GamepadButtonAux_LeftStickUp), FGamepadKeyNames::LeftStickUp },
+			{ static_cast<uint32>(EGameInputAuxButton::GamepadButtonAux_LeftStickDown), FGamepadKeyNames::LeftStickDown },
+			{ static_cast<uint32>(EGameInputAuxButton::GamepadButtonAux_LeftStickLeft), FGamepadKeyNames::LeftStickLeft },
+			{ static_cast<uint32>(EGameInputAuxButton::GamepadButtonAux_LeftStickRight), FGamepadKeyNames::LeftStickRight },
+			{ static_cast<uint32>(EGameInputAuxButton::GamepadButtonAux_RightStickUp), FGamepadKeyNames::RightStickUp },
+			{ static_cast<uint32>(EGameInputAuxButton::GamepadButtonAux_RightStickDown), FGamepadKeyNames::RightStickDown },
+			{ static_cast<uint32>(EGameInputAuxButton::GamepadButtonAux_RightStickLeft), FGamepadKeyNames::RightStickLeft },
+			{ static_cast<uint32>(EGameInputAuxButton::GamepadButtonAux_RightStickRight), FGamepadKeyNames::RightStickRight }
+		};
+		return GamepadButtonMap;
+	}
+}
+
+bool FGameInputGamepadDeviceProcessor::ProcessInput(const FGameInputEventParams& Params)
+{
+	// check if the reading had gamepad info
+	GameInputGamepadState GamepadState;
+	if (!Params.Reading->GetGamepadState(&GamepadState))
+	{
+		return false;
+	}
+
+	bool bRes = false;	
+
+	// We want to process gamepad BUTTON states for every game input reading.
+	bRes |= ProcessGamepadButtonState(Params, GamepadState);
+	++NumReadingsProcessedThisFrame;
+
+	return bRes;
+}
+
+bool FGameInputGamepadDeviceProcessor::PostProcessInput(const FGameInputEventParams& Params)
+{	
+	// On the last input reading for the frame, the "Current Reading" should always be null. we only care for the LastReading here
+	ensure (Params.Reading == nullptr);
+
+	const bool bProcessedAnyGamepadButtonsThisFrame = (NumReadingsProcessedThisFrame > 0);
+
+	// This is the last reading of this frame, reset the counter to 0
+	NumReadingsProcessedThisFrame = 0;
+
+	if (!Params.PreviousReading)
+	{
+		return false;
+	}
+
+	// Get the gamepad state of the *Last Reading* of the frame here
+	GameInputGamepadState GamepadState;
+	if (!Params.PreviousReading->GetGamepadState(&GamepadState))
+	{
+		return false;
+	}
+
+	bool bRes = false;
+
+	// If there were no gamepad events this frame, send button events using the previous frame's reading for button repeats.
+	// This is necessary because GetNextReading won't return any reading if the state is unchanged
+	if (!bProcessedAnyGamepadButtonsThisFrame)
+	{
+		bRes |= ProcessGamepadButtonState(Params, GamepadState);
+	}
+
+	// We only want to process gamepad ANALOG inputs on the last reading because we only care about
+	// the most recent analog stick input value. If there are multiple readings and we processed analog
+	// for every one, then the values would accumulate and stack, giving us incorrect data in the message handler.
+	bRes |= ProcessGamepadAnalogState(Params, GamepadState);
+
+	return bRes;
+}
+
+bool FGameInputGamepadDeviceProcessor::ProcessGamepadAnalogState(const FGameInputEventParams& Params, GameInputGamepadState& GamepadState)
+{
+	if (!Params.PlatformUserId.IsValid())
+	{
+		return false;
+	}
+
+	// ignore this input if the reading has remained empty from last time
+	const bool bIsEmptyReading = UE::GameInput::IsEmptyGamepadReading(GamepadState);
+	const bool bWasEmptyReading = UE::GameInput::IsEmptyGamepadReading(PreviousState);
+	const bool bHasDifferentAnalogInput = UE::GameInput::HasDifferentAnalogInput(GamepadState, PreviousState);
+	if (bIsEmptyReading && bWasEmptyReading && !bHasDifferentAnalogInput)
+	{
+		return false;
+	}
+
+	OnControllerAnalog(Params, FGamepadKeyNames::LeftAnalogX, GamepadState.leftThumbstickX, PreviousState.leftThumbstickX, UE::GameInput::GamepadLeftStickDeadzone);
+	OnControllerAnalog(Params, FGamepadKeyNames::LeftAnalogY, GamepadState.leftThumbstickY, PreviousState.leftThumbstickY, UE::GameInput::GamepadLeftStickDeadzone);
+	OnControllerAnalog(Params, FGamepadKeyNames::RightAnalogX, GamepadState.rightThumbstickX, PreviousState.rightThumbstickX, UE::GameInput::GamepadRightStickDeadzone);
+	OnControllerAnalog(Params, FGamepadKeyNames::RightAnalogY, GamepadState.rightThumbstickY, PreviousState.rightThumbstickY, UE::GameInput::GamepadRightStickDeadzone);
+	OnControllerAnalog(Params, FGamepadKeyNames::LeftTriggerAnalog, GamepadState.leftTrigger, PreviousState.leftTrigger, UE::GameInput::GamepadTriggerDeadzone);
+	OnControllerAnalog(Params, FGamepadKeyNames::RightTriggerAnalog, GamepadState.rightTrigger, PreviousState.rightTrigger, UE::GameInput::GamepadTriggerDeadzone);
+
+	// update saved analog state
+	PreviousState.leftTrigger = GamepadState.leftTrigger;
+	PreviousState.rightTrigger = GamepadState.rightTrigger;
+	PreviousState.leftThumbstickX = GamepadState.leftThumbstickX;
+	PreviousState.leftThumbstickY = GamepadState.leftThumbstickY;
+	PreviousState.rightThumbstickX = GamepadState.rightThumbstickX;
+	PreviousState.rightThumbstickY = GamepadState.rightThumbstickY;
+		
+	return true;
+}
+
+bool FGameInputGamepadDeviceProcessor::ProcessGamepadButtonState(const FGameInputEventParams& Params, GameInputGamepadState& GamepadState)
+{
+	if (!Params.PlatformUserId.IsValid())
+	{
+		return false;
+	}
+
+	// ignore this input if the reading has remained empty from last time, or we still have outstanding held buttons 
+	// (these held buttons are likely to occur with mapped analog triggers when there have been multiple input events this frame)
+	const bool bIsEmptyReading = UE::GameInput::IsEmptyGamepadReading(GamepadState);
+	const bool bWasEmptyReading = UE::GameInput::IsEmptyGamepadReading(PreviousState);
+	const bool bHasDifferentAnalogInput = UE::GameInput::HasDifferentAnalogInput(GamepadState, PreviousState);
+	if (bIsEmptyReading && bWasEmptyReading && !bHasDifferentAnalogInput && LastButtonHeldMask == 0)
+	{
+		return false;
+	}
+
+	// map buttons (low 15 bits)
+	uint32 CurrentButtonHeldMask = (static_cast<uint32>(GamepadState.buttons) & UE::GameInput::GamingInputButtonMask);
+
+	// map analog triggers to digital input.
+	if (GamepadState.leftTrigger > UE::GameInput::GamepadTriggerDeadzone)
+	{
+		CurrentButtonHeldMask |= UE::GameInput::EGameInputAuxButton::GamepadButtonAux_LeftTrigger;
+	}
+
+	if (GamepadState.rightTrigger > UE::GameInput::GamepadTriggerDeadzone)
+	{
+		CurrentButtonHeldMask |= UE::GameInput::EGameInputAuxButton::GamepadButtonAux_RightTrigger;
+	}
+
+	// map left/right stick digital inputs to (top 8 bits)
+	if (GamepadState.leftThumbstickY > UE::GameInput::GamepadLeftStickDeadzone)
+	{
+		CurrentButtonHeldMask |= UE::GameInput::EGameInputAuxButton::GamepadButtonAux_LeftStickUp;
+	}
+	else if (GamepadState.leftThumbstickY < -UE::GameInput::GamepadLeftStickDeadzone)
+	{
+		CurrentButtonHeldMask |= UE::GameInput::EGameInputAuxButton::GamepadButtonAux_LeftStickDown;
+	}
+	if (GamepadState.leftThumbstickX > UE::GameInput::GamepadLeftStickDeadzone)
+	{
+		CurrentButtonHeldMask |= UE::GameInput::EGameInputAuxButton::GamepadButtonAux_LeftStickRight;
+	}
+	else if (GamepadState.leftThumbstickX < -UE::GameInput::GamepadLeftStickDeadzone)
+	{
+		CurrentButtonHeldMask |= UE::GameInput::EGameInputAuxButton::GamepadButtonAux_LeftStickLeft;
+	}
+
+	if (GamepadState.rightThumbstickY > UE::GameInput::GamepadRightStickDeadzone)
+	{
+		CurrentButtonHeldMask |= UE::GameInput::EGameInputAuxButton::GamepadButtonAux_RightStickUp;
+	}
+	else if (GamepadState.rightThumbstickY < -UE::GameInput::GamepadRightStickDeadzone)
+	{
+		CurrentButtonHeldMask |= UE::GameInput::EGameInputAuxButton::GamepadButtonAux_RightStickDown;
+	}
+	if (GamepadState.rightThumbstickX > UE::GameInput::GamepadRightStickDeadzone)
+	{
+		CurrentButtonHeldMask |= UE::GameInput::EGameInputAuxButton::GamepadButtonAux_RightStickRight;
+	}
+	else if (GamepadState.rightThumbstickX < -UE::GameInput::GamepadRightStickDeadzone)
+	{
+		CurrentButtonHeldMask |= UE::GameInput::EGameInputAuxButton::GamepadButtonAux_RightStickLeft;
+	}
+
+	EvaluateButtonStates(
+		Params,
+		CurrentButtonHeldMask,
+		OUT LastButtonHeldMask,
+		RepeatTime,
+		UE::GameInput::GetGamepadButtonMap(),
+		MaxSupportedButtons);
+
+	// Keep track of the current BUTTON state. We don't want to update the entire PreviousState struct here
+	// because buttons may be evaluated more then analog inputs per-frame
+	PreviousState.buttons = GamepadState.buttons;
+
+	return true;
+}
+
+void FGameInputGamepadDeviceProcessor::ClearState(const FGameInputEventParams& Params)
+{
+	// We need a valid input device id when sending messages, otherwise slate will hit a check
+	// and attempt to create some new slate user with an invalid index.
+	if (!Params.PlatformUserId.IsValid() || !Params.InputDeviceId.IsValid())
+	{
+		return;
+	}
+
+	// Reset Axis values
+	OnControllerAnalog(Params, FGamepadKeyNames::LeftAnalogX, 0.0f, PreviousState.leftThumbstickX, UE::GameInput::GamepadLeftStickDeadzone);
+	OnControllerAnalog(Params, FGamepadKeyNames::LeftAnalogY, 0.0f, PreviousState.leftThumbstickY, UE::GameInput::GamepadLeftStickDeadzone);
+	OnControllerAnalog(Params, FGamepadKeyNames::RightAnalogX, 0.0f, PreviousState.rightThumbstickX, UE::GameInput::GamepadRightStickDeadzone);
+	OnControllerAnalog(Params, FGamepadKeyNames::RightAnalogY, 0.0f, PreviousState.rightThumbstickY, UE::GameInput::GamepadRightStickDeadzone);
+	OnControllerAnalog(Params, FGamepadKeyNames::LeftTriggerAnalog, 0.0f, PreviousState.leftTrigger, UE::GameInput::GamepadTriggerDeadzone);
+	OnControllerAnalog(Params, FGamepadKeyNames::RightTriggerAnalog, 0.0f, PreviousState.rightTrigger, UE::GameInput::GamepadTriggerDeadzone);
+
+	// Reset button values
+
+	// Just use 0 as our button mask because we want them all to be set to 0
+	uint32 CurrentButtonHeldMask = 0;
+
+	EvaluateButtonStates(
+		Params,
+		CurrentButtonHeldMask,
+		LastButtonHeldMask,
+		RepeatTime,
+		UE::GameInput::GetGamepadButtonMap(),
+		MaxSupportedButtons);
+
+	for (uint32 i = 0; i < MaxSupportedButtons; i++)
+	{
+		RepeatTime[i] = 0.0;
+	}
+
+	// clear previous gamepad state
+	FMemory::Memset(PreviousState, 0);
+}
+
+GameInputKind FGameInputGamepadDeviceProcessor::GetSupportedReadingKind() const
+{
+	return GameInputKindGamepad;
+}
+
+////////////////////////////////////////////////////////
+// FGameInputControllerDeviceProcessor
+
+/**
+* Returns a map of GameInputLabel's to their associated Unreal Engine FKey names.
+*/
+const TMap<GameInputLabel, FGamepadKeyNames::Type>& FGameInputControllerDeviceProcessor::GetGameInputButtonLabelToUnrealName()
+{
+	static const TMap<GameInputLabel, FGamepadKeyNames::Type> LabelMap
+	{
+		{ GameInputLabelUnknown, FGamepadKeyNames::Invalid },
+		{ GameInputLabelNone, FGamepadKeyNames::Invalid },
+		{ GameInputLabelXboxGuide, FGamepadKeyNames::SpecialRight },	// TODO: Check if this is correct
+		{ GameInputLabelXboxBack, FGamepadKeyNames::SpecialLeft },
+		{ GameInputLabelXboxStart, FGamepadKeyNames::SpecialRight },
+		{ GameInputLabelXboxMenu, FGamepadKeyNames::SpecialRight },
+		{ GameInputLabelXboxView, FGamepadKeyNames::SpecialLeft },
+		{ GameInputLabelXboxA, FGamepadKeyNames::FaceButtonBottom },
+		{ GameInputLabelXboxB, FGamepadKeyNames::FaceButtonRight },
+		{ GameInputLabelXboxX, FGamepadKeyNames::FaceButtonLeft },
+		{ GameInputLabelXboxY, FGamepadKeyNames::FaceButtonTop },
+		{ GameInputLabelXboxDPadUp, FGamepadKeyNames::DPadUp },
+		{ GameInputLabelXboxDPadDown, FGamepadKeyNames::DPadDown },
+		{ GameInputLabelXboxDPadLeft, FGamepadKeyNames::DPadLeft },
+		{ GameInputLabelXboxDPadRight, FGamepadKeyNames::DPadRight },
+		{ GameInputLabelXboxLeftShoulder, FGamepadKeyNames::LeftShoulder },
+		{ GameInputLabelXboxLeftTrigger, FGamepadKeyNames::LeftTriggerAnalog },
+		{ GameInputLabelXboxLeftStickButton, FGamepadKeyNames::LeftThumb },
+		{ GameInputLabelXboxRightShoulder, FGamepadKeyNames::RightShoulder },
+		{ GameInputLabelXboxRightTrigger, FGamepadKeyNames::RightTriggerAnalog },
+		{ GameInputLabelXboxRightStickButton, FGamepadKeyNames::RightThumb },
+		{ GameInputLabelXboxPaddle1, FGamepadKeyNames::Invalid },				// TODO: Do we need special additional FKey's for these paddle types?
+		{ GameInputLabelXboxPaddle2, FGamepadKeyNames::Invalid },				// Return invalid for now, but I thought the Xbox One pro controller would
+		{ GameInputLabelXboxPaddle3, FGamepadKeyNames::Invalid },				// be handled by the OS itself via virtual remapping
+		{ GameInputLabelXboxPaddle4, FGamepadKeyNames::Invalid },
+		{ GameInputLabelLetterA, EKeys::A.GetFName() },
+		{ GameInputLabelLetterB, EKeys::B.GetFName() },
+		{ GameInputLabelLetterC, EKeys::C.GetFName() },
+		{ GameInputLabelLetterD, EKeys::D.GetFName() },
+		{ GameInputLabelLetterE, EKeys::E.GetFName() },
+		{ GameInputLabelLetterF, EKeys::F.GetFName() },
+		{ GameInputLabelLetterG, EKeys::G.GetFName() },
+		{ GameInputLabelLetterH, EKeys::H.GetFName() },
+		{ GameInputLabelLetterI, EKeys::I.GetFName() },
+		{ GameInputLabelLetterJ, EKeys::J.GetFName() },
+		{ GameInputLabelLetterK, EKeys::K.GetFName() },
+		{ GameInputLabelLetterL, EKeys::L.GetFName() },
+		{ GameInputLabelLetterM, EKeys::M.GetFName() },
+		{ GameInputLabelLetterN, EKeys::N.GetFName() },
+		{ GameInputLabelLetterO, EKeys::O.GetFName() },
+		{ GameInputLabelLetterP, EKeys::P.GetFName() },
+		{ GameInputLabelLetterQ, EKeys::Q.GetFName() },
+		{ GameInputLabelLetterR, EKeys::R.GetFName() },
+		{ GameInputLabelLetterS, EKeys::S.GetFName() },
+		{ GameInputLabelLetterT, EKeys::T.GetFName() },
+		{ GameInputLabelLetterU, EKeys::U.GetFName() },
+		{ GameInputLabelLetterV, EKeys::V.GetFName() },
+		{ GameInputLabelLetterW, EKeys::W.GetFName() },
+		{ GameInputLabelLetterX, EKeys::X.GetFName() },
+		{ GameInputLabelLetterY, EKeys::Y.GetFName() },
+		{ GameInputLabelLetterZ, EKeys::Z.GetFName() },
+		{ GameInputLabelNumber0, EKeys::Zero.GetFName() },
+		{ GameInputLabelNumber1, EKeys::One.GetFName() },
+		{ GameInputLabelNumber2, EKeys::Two.GetFName() },
+		{ GameInputLabelNumber3, EKeys::Three.GetFName() },
+		{ GameInputLabelNumber4, EKeys::Four.GetFName() },
+		{ GameInputLabelNumber5, EKeys::Five.GetFName() },
+		{ GameInputLabelNumber6, EKeys::Six.GetFName() },
+		{ GameInputLabelNumber7, EKeys::Seven.GetFName() },
+		{ GameInputLabelNumber8, EKeys::Eight.GetFName() },
+		{ GameInputLabelNumber9, EKeys::Nine.GetFName() },
+		{ GameInputLabelArrowUp, EKeys::Up.GetFName() },
+		{ GameInputLabelArrowUpRight, EKeys::Up.GetFName() },
+		{ GameInputLabelArrowRight, EKeys::Right.GetFName() },
+		{ GameInputLabelArrowDownRight, EKeys::Down.GetFName() },	// TODO: We should support multiple FKey's here, like we do for switches
+		{ GameInputLabelArrowDown, EKeys::Down.GetFName() },
+		{ GameInputLabelArrowDownLLeft, EKeys::Down.GetFName() },
+		{ GameInputLabelArrowLeft, EKeys::Left.GetFName() },
+		{ GameInputLabelArrowUpLeft, FGamepadKeyNames::DPadUp },
+		{ GameInputLabelArrowUpDown, FGamepadKeyNames::DPadUp },
+		{ GameInputLabelArrowLeftRight, FGamepadKeyNames::DPadUp },
+		{ GameInputLabelArrowUpDownLeftRight, FGamepadKeyNames::DPadUp },
+		{ GameInputLabelArrowClockwise, FGamepadKeyNames::DPadUp },		// TODO: new key for this          
+		{ GameInputLabelArrowCounterClockwise, FGamepadKeyNames::DPadUp },	// TODO: new key for this   
+		{ GameInputLabelArrowReturn, EKeys::Enter.GetFName() },
+		{ GameInputLabelIconBranding, EKeys::Home.GetFName() },		// TODO: I dont think we have a UE key for this, maybe we use home?     
+		{ GameInputLabelIconHome, FGamepadKeyNames::SpecialRight },
+		{ GameInputLabelIconMenu, FGamepadKeyNames::SpecialLeft },
+		{ GameInputLabelIconCross, FGamepadKeyNames::FaceButtonBottom },
+		{ GameInputLabelIconCircle, FGamepadKeyNames::FaceButtonRight },
+		{ GameInputLabelIconSquare, FGamepadKeyNames::FaceButtonLeft },
+		{ GameInputLabelIconTriangle, FGamepadKeyNames::FaceButtonTop },
+		{ GameInputLabelIconStar, EKeys::Asterix.GetFName() },    // TODO: Star? Is this the asterix?            
+		{ GameInputLabelIconDPadUp, FGamepadKeyNames::DPadUp },
+		{ GameInputLabelIconDPadDown, FGamepadKeyNames::DPadDown },
+		{ GameInputLabelIconDPadLeft, FGamepadKeyNames::DPadLeft },
+		{ GameInputLabelIconDPadRight, FGamepadKeyNames::DPadRight },
+		{ GameInputLabelIconDialClockwise, FGamepadKeyNames::DPadUp },
+		{ GameInputLabelIconDialCounterClockwise, FGamepadKeyNames::DPadUp },
+		{ GameInputLabelIconSliderLeftRight, FGamepadKeyNames::DPadUp },
+		{ GameInputLabelIconSliderUpDown, FGamepadKeyNames::DPadUp },
+		{ GameInputLabelIconWheelUpDown, FGamepadKeyNames::DPadUp },
+		{ GameInputLabelIconPlus, EKeys::Add.GetFName() },
+		{ GameInputLabelIconMinus, EKeys::Subtract.GetFName() },
+		{ GameInputLabelIconSuspension, FGamepadKeyNames::DPadUp },
+		{ GameInputLabelHome, EKeys::Home.GetFName() },			// TODO: Do we have a gamepad key for guide?            
+		{ GameInputLabelGuide, FGamepadKeyNames::SpecialLeft },
+		{ GameInputLabelMode, FGamepadKeyNames::SpecialLeft },
+		{ GameInputLabelSelect, FGamepadKeyNames::SpecialRight },
+		{ GameInputLabelMenu, FGamepadKeyNames::SpecialRight },
+		{ GameInputLabelView, FGamepadKeyNames::SpecialLeft },
+		{ GameInputLabelBack, FGamepadKeyNames::SpecialLeft },
+		{ GameInputLabelStart, FGamepadKeyNames::SpecialRight },
+		{ GameInputLabelOptions, FGamepadKeyNames::SpecialRight },
+		{ GameInputLabelShare, FGamepadKeyNames::SpecialLeft },
+		{ GameInputLabelUp, FGamepadKeyNames::DPadUp },
+		{ GameInputLabelDown, FGamepadKeyNames::DPadDown },
+		{ GameInputLabelLeft, FGamepadKeyNames::DPadLeft },
+		{ GameInputLabelRight, FGamepadKeyNames::DPadRight },
+		{ GameInputLabelLB, FGamepadKeyNames::LeftShoulder },
+		{ GameInputLabelLT, FGamepadKeyNames::LeftTriggerAnalog },
+		{ GameInputLabelLSB, FGamepadKeyNames::LeftShoulder },
+		{ GameInputLabelL1, FGamepadKeyNames::LeftShoulder },
+		{ GameInputLabelL2, FGamepadKeyNames::LeftTriggerAnalog },
+		{ GameInputLabelL3, FGamepadKeyNames::DPadUp },
+		{ GameInputLabelRB, FGamepadKeyNames::RightShoulder },
+		{ GameInputLabelRT, FGamepadKeyNames::RightTriggerAnalog },
+		{ GameInputLabelRSB, FGamepadKeyNames::RightShoulder },
+		{ GameInputLabelR1, FGamepadKeyNames::RightShoulder },
+		{ GameInputLabelR2, FGamepadKeyNames::RightTriggerAnalog },
+		{ GameInputLabelR3, FGamepadKeyNames::DPadUp },
+		{ GameInputLabelP1, FGamepadKeyNames::DPadUp },
+		{ GameInputLabelP2, FGamepadKeyNames::DPadUp },						// what are these? More paddle types?     
+		{ GameInputLabelP3, FGamepadKeyNames::DPadUp },
+		{ GameInputLabelP4, FGamepadKeyNames::DPadUp }
+	};
+
+	return LabelMap;
+}
+
+FGameInputControllerDeviceProcessor::FGameInputControllerDeviceProcessor()
+	: IGameInputDeviceProcessor()
+{
+	// Ensure that our switch repeat times array has some default values set on it by default
+	// because we access it with the [] operator when processing it. 
+	SwitchRepeatTimes.AddDefaulted(static_cast<uint32>(GameInputSwitchLeft) + 1);
+}
+
+bool FGameInputControllerDeviceProcessor::ProcessInput(const FGameInputEventParams& Params)
+{
+	bool bRes = false;
+
+	const FGameInputDeviceConfiguration* ControllerConfig = GetDefault<UGameInputDeveloperSettings>()->FindDeviceConfiguration(Params.GetDeviceInfo());
+
+	// Note that we use the current reading here. "ProcessInput" can be called multiple times per frame
+	// if there is more then one input reading in the stack. We want to process all button and switch states
+	// to ensure that we don't miss one
+	bRes |= ProcessControllerSwitchState(Params, ControllerConfig, Params.Reading);
+	bRes |= ProcessControllerButtonState(Params, ControllerConfig, Params.Reading);
+	++NumReadingsProcessedThisFrame;
+
+	return bRes;
+}
+
+bool FGameInputControllerDeviceProcessor::PostProcessInput(const FGameInputEventParams& Params)
+{
+	// On the last input reading for the frame, the "Current Reading" should always be null. we only care for the LastReading here
+	ensure(Params.Reading == nullptr);
+
+	bool bRes = false;
+
+	const FGameInputDeviceConfiguration* ControllerConfig = GetDefault<UGameInputDeveloperSettings>()->FindDeviceConfiguration(Params.GetDeviceInfo());
+
+	const bool bProcessedAnyButtonsThisFrame = (NumReadingsProcessedThisFrame > 0);
+
+	// This is the last reading of this frame, reset the counter to 0
+	NumReadingsProcessedThisFrame = 0;
+
+	// Note that we use "Params.PreviousReading", because the current reading will be null. 
+	if (!bProcessedAnyButtonsThisFrame)
+	{
+		bRes |= ProcessControllerSwitchState(Params, ControllerConfig, Params.PreviousReading);
+		bRes |= ProcessControllerButtonState(Params, ControllerConfig, Params.PreviousReading);
+	}
+
+	bRes |= ProcessControllerAxisState(Params, ControllerConfig, Params.PreviousReading);
+
+	return bRes;
+}
+
+void FGameInputControllerDeviceProcessor::ClearState(const FGameInputEventParams& Params)
+{
+	// Reset the switches to all be "center" positions
+	{
+		const uint32 SwitchCount = static_cast<uint32>(PreviousSwitchPositions.Num());
+
+		TArray<GameInputSwitchPosition> SwitchPositions;
+		SwitchPositions.AddUninitialized(SwitchCount);
+		// Evaluate all switch positions
+		for (uint32 i = 0; i < SwitchCount; ++i)
+		{
+			// The PreviousSwitchPositions will be updated to the current switch position by the evaluate function
+			EvaluateSwitchState(
+				Params,
+				SwitchPositions[i],
+				OUT PreviousSwitchPositions[i],
+				SwitchRepeatTimes);
+		}
+	}
+	
+
+	if (const FGameInputDeviceConfiguration* Config = GetDefault<UGameInputDeveloperSettings>()->FindDeviceConfiguration(Params.GetDeviceInfo()))
+	{
+		// Reset the button state to not being pressed
+		// Just use 0 as our button mask because we want them all to be set to 0
+		uint32 CurrentButtonHeldMask = 0;
+
+		EvaluateButtonStates(
+			Params,
+			CurrentButtonHeldMask,
+			LastButtonHeldMask,
+			RepeatTime,
+			Config->ControllerButtonMappingData,
+			MaxSupportedButtons);
+			
+		// Reset the axis state to be zero on all axis that we had in the previous state
+		for (int32 i = 0; i < PreviousControllerAxisValues.Num(); ++i)
+		{
+			if (const FGameInputControllerAxisData* AxisData = Config->ControllerAxisMappingData.Find(i))
+			{
+				OnControllerAnalog(Params, AxisData->KeyName, 0.0f, PreviousControllerAxisValues[i], AxisData->DeadZone);
+			}			
+		}
+
+		FMemory::Memset(PreviousControllerAxisValues, 0);
+	}
+}
+
+GameInputKind FGameInputControllerDeviceProcessor::GetSupportedReadingKind() const
+{
+	return GameInputKindController | GameInputKindControllerAxis | GameInputKindControllerButton;
+}
+
+bool FGameInputControllerDeviceProcessor::ProcessControllerButtonState(const FGameInputEventParams& Params, const FGameInputDeviceConfiguration* ControllerConfig, IGameInputReading* InputReading)
+{
+	// We can only process generic controllers that have their config set up
+	// but allow you to run without a config to log out button indexes to make it 
+	// easier to configure your device
+	if (ControllerConfig && !ControllerConfig->bProcessControllerButtons)
+	{
+		return false;
+	}
+
+	const uint32 ButtonCount = InputReading->GetControllerButtonCount();
+
+	TArray<bool> ButtonStates;
+	ButtonStates.AddUninitialized(ButtonCount);
+
+	const uint32 Res = InputReading->GetControllerButtonState(ButtonCount, ButtonStates.GetData());
+	if (!Res)
+	{
+		return false;
+	}
+	
+	uint32 CurrentButtonHeldMask = 0x00;
+
+	for (uint32 i = 0; i < ButtonCount; ++i)
+	{
+		const FName* KeyName = ControllerConfig ? ControllerConfig->ControllerButtonMappingData.Find(1 << i) : nullptr;
+
+		CurrentButtonHeldMask |= (ButtonStates[i] << i);
+
+		if (ButtonStates[i])
+		{
+			UE_LOG(LogGameInput, Verbose, TEXT("[ProcessControllerButtonState] Device ID: %d   Button count: %d    index: %d  State: %d   %s"), 
+				Params.InputDeviceId.GetId(),
+				ButtonCount, 
+				i,
+				ButtonStates[i], 
+				KeyName ? *(KeyName->ToString()) : TEXT("NONE"));
+		}
+	}
+
+	if (ControllerConfig)
+	{
+		EvaluateButtonStates(
+			Params,
+			CurrentButtonHeldMask,
+			OUT LastButtonHeldMask,
+			RepeatTime,
+			ControllerConfig->ControllerButtonMappingData,
+			ButtonCount
+		);
+	}	
+
+	// Note: Ideally we would try and use the GameInput label system, 
+	// but it is currently unfinished within the GameInput API itself so we can't.
+	// We have provided this device specific config driven option instead.
+
+	return true;
+}
+
+bool FGameInputControllerDeviceProcessor::ProcessControllerAxisState(const FGameInputEventParams& Params, const FGameInputDeviceConfiguration* Config, IGameInputReading* InputReading)
+{
+	if (Config && !Config->bProcessControllerAxis)
+	{
+		return false;
+	}
+
+	const uint32 AxisCount = InputReading->GetControllerAxisCount();
+	TArray<float> AxisValues;
+	AxisValues.AddZeroed(AxisCount);
+	if (!InputReading->GetControllerAxisState(AxisCount, AxisValues.GetData()))
+	{
+		return false;
+	}
+
+	// Make sure that there is previous controller data initialized if necessary
+	while (PreviousControllerAxisValues.Num() < static_cast<int32>(AxisCount))
+	{
+		PreviousControllerAxisValues.AddZeroed();
+	}
+
+	for (uint32 i = 0; i < AxisCount; ++i)
+	{
+		float CurrentValue = AxisValues[i];
+		const float PreviousValue = PreviousControllerAxisValues[i];
+
+		if (Config)
+		{
+			if (const FGameInputControllerAxisData* AxisData = Config->ControllerAxisMappingData.Find(i))
+			{
+				if (AxisData->KeyName.IsValid() && AxisData->KeyName != NAME_None)
+				{
+					if (AxisData->bIsPackedPositveAndNegative)
+					{
+						// Maps the value to be -1.0 to +1.0
+						CurrentValue = (CurrentValue * 2.f) - 1.f;
+					}
+
+					CurrentValue *= AxisData->Scalar;
+					
+					OnControllerAnalog(Params, AxisData->KeyName, CurrentValue, PreviousValue, AxisData->DeadZone);
+
+					// Store this value for the next frame to compare to
+					PreviousControllerAxisValues[i] = CurrentValue;
+				}
+				else
+				{
+					UE_LOG(LogGameInput, VeryVerbose, TEXT("[ProcessControllerAxisState] (Device %s) Invalid key name configured for controller axis '%d': %.3f"), *UE::GameInput::LexToString(Params.Device), i, CurrentValue);
+				}
+			}
+		}		
+		// You are receiving analog values from an axis that you might not know about, log it here
+		// in case you are trying to set something up
+		else
+		{
+			UE_LOG(LogGameInput, VeryVerbose, TEXT("[ProcessControllerAxisState] (Device %s) Receiving input from an unconfigured controller axis '%d': %.3f"), *UE::GameInput::LexToString(Params.Device), i, CurrentValue);
+		}
+	}
+
+	return false;
+}
+
+bool FGameInputControllerDeviceProcessor::ProcessControllerSwitchState(const FGameInputEventParams& Params, const FGameInputDeviceConfiguration* Config, IGameInputReading* InputReading)
+{
+	if (Config && !Config->bProcessControllerSwitchState)
+	{
+		return false;
+	}
+
+	const uint32 SwitchCount = InputReading->GetControllerSwitchCount();
+	TArray<GameInputSwitchPosition> SwitchPositions;
+	SwitchPositions.AddUninitialized(SwitchCount);
+
+	// Make sure that we have some previous state initialized if we can
+	if (static_cast<uint32>(PreviousSwitchPositions.Num()) < SwitchCount)
+	{
+		PreviousSwitchPositions.AddDefaulted(SwitchCount);
+	}
+
+	uint32 Res = InputReading->GetControllerSwitchState(SwitchCount, SwitchPositions.GetData());
+
+	if (!Res)
+	{
+		return false;
+	}
+
+	// Evaluate all switch positions
+	for (uint32 i = 0; i < SwitchCount; ++i)
+	{
+		// The PreviousSwitchPositions will be updated to the current switch position by the evaluate function
+		EvaluateSwitchState(
+			Params,
+			SwitchPositions[i],
+			OUT PreviousSwitchPositions[i],
+			SwitchRepeatTimes);
+	}
+
+	return true;
+}
+
+////////////////////////////////////////////////////////
+// FGameInputKeyboardDeviceProcessor
+
+static float GKeyboardRepeatInitialDelay = 0.25f;
+static FAutoConsoleVariableRef CVarKeyboardRepeatInitialDelay(
+	TEXT("Input.KeyboardRepeatInitialDelay"),
+	GKeyboardRepeatInitialDelay,
+	TEXT("Time in seconds before a key repeat starts"));
+
+static float GKeyboardRepeatDelay = 0.05f;
+static FAutoConsoleVariableRef CVarKeyboardRepeatDelay(
+	TEXT("Input.KeyboardRepeatDelay"),
+	GKeyboardRepeatDelay,
+	TEXT("Time in seconds between each subsequent key repeat"));
+
+bool FGameInputKeyboardDeviceProcessor::ProcessInput(const FGameInputEventParams& Params)
+{	
+	TSet<uint8> CurrentPressedKeys;
+
+	int32 KeyCount = Params.Reading->GetKeyCount();
+	if (KeyCount > 0)
+	{
+		// read the key state
+		TArray<GameInputKeyState> KeyStates;
+		KeyStates.AddUninitialized(KeyCount);
+		int32 ReadCount = Params.Reading->GetKeyState(KeyCount, KeyStates.GetData());
+
+		// build a set of the pressed keycodes
+		for (GameInputKeyState KeyState : KeyStates)
+		{
+			CurrentPressedKeys.Add(KeyState.virtualKey);
+		}		
+	}
+
+	UpdateUnifiedKeyboardState(Params, CurrentPressedKeys);
+	return true;
+}
+
+void FGameInputKeyboardDeviceProcessor::UpdateUnifiedKeyboardState(const FGameInputEventParams& Params, TSet<uint8>& CurrentPressedKeys)
+{
+	// process unified pressed keys
+	double CurrentTime = FPlatformTime::Seconds();
+	for (uint8 KeyCode : CurrentPressedKeys)
+	{
+		bool bIsRepeat = LastPressedKeys.Contains(KeyCode);
+		if (!bIsRepeat)
+		{
+			KeyRepeatTime.Add(KeyCode, CurrentTime + GKeyboardRepeatInitialDelay);
+			Params.MessageHandler->OnKeyDown(KeyCode, 0, false);
+			UE_LOG(LogGameInput, Verbose, TEXT("Key Press 0x%X"), KeyCode);
+
+			if (KeyCode == VK_CAPITAL)
+			{
+				SetSimulatedCapsLock(!bSimulatedCapsLock);
+				UE_LOG(LogGameInput, Verbose, TEXT("Simulated caps lock is %s"), bSimulatedCapsLock ? TEXT("ON") : TEXT("OFF"));
+			}
+		}
+		else if (CurrentTime > KeyRepeatTime[KeyCode])
+		{
+			KeyRepeatTime.Add(KeyCode, CurrentTime + GKeyboardRepeatDelay);
+			Params.MessageHandler->OnKeyDown(KeyCode, 0, true);
+			UE_LOG(LogGameInput, Verbose, TEXT("Key Press 0x%X (repeat)"), KeyCode);
+		}
+	}
+
+	// process any released keys
+	for (uint8 KeyCode : LastPressedKeys)
+	{
+		if (!CurrentPressedKeys.Contains(KeyCode))
+		{
+			KeyRepeatTime.Remove(KeyCode);
+			Params.MessageHandler->OnKeyUp(KeyCode, 0, false);
+			UE_LOG(LogGameInput, Verbose, TEXT("Key Release 0x%X"), KeyCode);
+		}
+	}
+
+	// update saved state
+	LastPressedKeys = CurrentPressedKeys;
+
+}
+
+void FGameInputKeyboardDeviceProcessor::SetSimulatedCapsLock(bool bVal)
+{
+	bSimulatedCapsLock = bVal;
+}
+
+void FGameInputKeyboardDeviceProcessor::ClearState(const FGameInputEventParams& Params)
+{
+	TSet<uint8> NoPressedKeys;
+	UpdateUnifiedKeyboardState(Params, NoPressedKeys);
+}
+
+GameInputKind FGameInputKeyboardDeviceProcessor::GetSupportedReadingKind() const
+{
+	return GameInputKindKeyboard;
+}
+
+////////////////////////////////////////////////////////
+// FGameInputMouseDeviceProcessor
+
+static int GAllowVirtualMouseInput = 1;
+static FAutoConsoleVariableRef CVarAllowVirtualMouseInput(
+	TEXT("GameInput.AllowVirtualMouseInput"),
+	GAllowVirtualMouseInput,
+	TEXT("Whether to accept input from virtual mice, such as those from a remote viewer. Note that this doesn't change whether the mouse is 'connected'"));
+
+static float GMouseSensitivity = 1.0f;
+static FAutoConsoleVariableRef CVarMouseSensitivity(
+	TEXT("Input.MouseSensitivity"),
+	GMouseSensitivity,
+	TEXT("The sensitivity multiplier of the mouse\n")
+	TEXT(" 1 (default)"),
+	ECVF_Default);
+
+static float GMouseDoubleClickArea = 10.0f;
+static FAutoConsoleVariableRef CVarMouseDoubleClickArea(
+	TEXT("Input.DoubleClickArea"),
+	GMouseDoubleClickArea,
+	TEXT("How far the mouse can move between double clicks to still count as a double click"));
+
+static float GMouseDoubleClickDelay = 0.5f;
+static FAutoConsoleVariableRef CVarMouseDoubleClickDelay(
+	TEXT("Input.DoubleClickDelay"),
+	GMouseDoubleClickDelay,
+	TEXT("Time in seconds between mouse down events to trigger a double click event. Set to 0 to disable double clicking"));
+
+namespace UE::GameInput
+{
+	struct FMouseButtonMapping
+	{
+		uint32 GameInputButton;
+		EMouseButtons::Type MouseButton;
+	};
+	static const FMouseButtonMapping MouseButtonMappings[] =
+	{
+		{ static_cast<uint32>(GameInputMouseButtons::GameInputMouseLeftButton), EMouseButtons::Left    },
+		{ static_cast<uint32>(GameInputMouseButtons::GameInputMouseRightButton), EMouseButtons::Right   },
+		{ static_cast<uint32>(GameInputMouseButtons::GameInputMouseMiddleButton), EMouseButtons::Middle  },
+		{ static_cast<uint32>(GameInputMouseButtons::GameInputMouseButton4), EMouseButtons::Thumb01 },
+		{ static_cast<uint32>(GameInputMouseButtons::GameInputMouseButton5), EMouseButtons::Thumb02 },
+	};
+}
+
+FGameInputMouseDeviceProcessor::FGameInputMouseDeviceProcessor(const TSharedPtr<class ICursor>& InCursor)
+	: Cursor(InCursor)
+	, LastMouseOffset(ForceInitToZero)
+{
+	FMemory::Memset(PreviousMouseState, 0);
+
+	for (uint32 i = 0; i < MaxSupportedButtons; i++)
+	{
+		RepeatTime[i] = 0.0;
+	}
+}
+
+float FGameInputMouseDeviceProcessor::CalculateMouseAcceleration(int32 Delta)
+{
+	const float NominalMovement = 20.0f;
+	const float Sign = (Delta > 0) ? 1.0f : -1.0f;
+	const float FDelta = (float)Delta;
+	return Sign * FMath::Max(FMath::Abs(FDelta), FMath::Abs(GMouseSensitivity * FDelta * FDelta / NominalMovement));
+}
+
+bool FGameInputMouseDeviceProcessor::CanProcessVirtualMouse() const
+{
+	// ignore the input if requested
+	if (GAllowVirtualMouseInput == 0)
+	{
+		return false;
+	}
+	return true;
+}
+
+bool FGameInputMouseDeviceProcessor::ProcessInput(const FGameInputEventParams& Params)
+{
+	// read mouse state
+	GameInputMouseState MouseState;
+	if (!Params.Reading->GetMouseState(&MouseState))
+	{
+		return false;
+	}
+
+	int32 LEGACY_VirtualMaxX = MAX_int32;
+	int32 LEGACY_VirtualMaxY = MAX_int32;
+
+	const GameInputDeviceInfo* DeviceInfo = Params.GetDeviceInfo();
+	const bool bIsVirtualMouse = DeviceInfo ? DeviceInfo->deviceFamily == GameInputFamilyVirtual : false;
+
+	const bool bHighPrecisionMouseMode = FSlateApplication::Get().IsUsingHighPrecisionMouseMovment();
+
+	if (bIsVirtualMouse)
+	{
+		if (!CanProcessVirtualMouse())
+		{
+			return false;
+		}
+	}
+
+	// update mouse position
+	float MouseDX, MouseDY = 0.0f;
+	if (bHighPrecisionMouseMode)
+	{
+		MouseDX = static_cast<float>(MouseState.positionX - PreviousMouseState.positionX);
+		MouseDY = static_cast<float>(MouseState.positionY - PreviousMouseState.positionY);
+	}
+	else
+	{
+		MouseDX = CalculateMouseAcceleration(static_cast<int32>(MouseState.positionX) - static_cast<int32>(PreviousMouseState.positionX));
+		MouseDY = CalculateMouseAcceleration(static_cast<int32>(MouseState.positionY) - static_cast<int32>(PreviousMouseState.positionY));
+	}
+	if (MouseDX != 0 || MouseDY != 0)
+	{
+		if (Cursor.IsValid())
+		{
+			FVector2D CursorPos = Cursor->GetPosition();
+			CursorPos.X += MouseDX;
+			CursorPos.Y += MouseDY;
+			Cursor->SetPosition((int32)CursorPos.X, (int32)CursorPos.Y);
+		}
+
+		LastMouseOffset.X += MouseDX;
+		LastMouseOffset.Y += MouseDY;
+
+		UE_LOG(LogGameInput, VeryVerbose, TEXT("Device %s (InputDeviceId = %d) - Mouse DX %0.2f, DY %.2f  %s"), *UE::GameInput::LexToString(Params.Device), Params.InputDeviceId.GetId(), MouseDX, MouseDY, bHighPrecisionMouseMode ? TEXT("(high precision)") : TEXT(""));
+		Params.MessageHandler->OnRawMouseMove((int32)MouseDX, (int32)MouseDY);
+	}
+
+	// update mouse wheel
+	const float MouseWheelSpinFactor = 1 / 120.0f;
+	float MouseWheelDX = (float)(MouseState.wheelX - PreviousMouseState.wheelX);
+	float MouseWheelDY = (float)(MouseState.wheelY - PreviousMouseState.wheelY);
+	if (MouseWheelDY != 0)
+	{
+		UE_LOG(LogGameInput, Verbose, TEXT("Device %s (InputDeviceId = %d) - Mouse Wheel DY %.2f"), *UE::GameInput::LexToString(Params.Device), Params.InputDeviceId.GetId(), MouseWheelDY);
+		Params.MessageHandler->OnMouseWheel(MouseWheelDY * MouseWheelSpinFactor);
+	}
+
+	// handle button change events
+	uint32 CurrentButtonHeldMask = (uint32)MouseState.buttons;
+	uint32 ActionMask = (LastButtonHeldMask ^ CurrentButtonHeldMask);
+	uint32 RepeatMask = (LastButtonHeldMask & CurrentButtonHeldMask);
+	double CurrentTime = FPlatformTime::Seconds();
+
+	for (int ButtonIndex = 0; ButtonIndex < UE_ARRAY_COUNT(UE::GameInput::MouseButtonMappings); ButtonIndex++)
+	{
+		const uint32 GameInputButton = UE::GameInput::MouseButtonMappings[ButtonIndex].GameInputButton;
+		if ((ActionMask & GameInputButton) == 0)
+		{
+			continue;
+		}
+
+		const EMouseButtons::Type MouseButton = UE::GameInput::MouseButtonMappings[ButtonIndex].MouseButton;
+		if ((CurrentButtonHeldMask & GameInputButton) != 0)
+		{
+			const bool bHasDoubleClick = (GMouseDoubleClickDelay > 0) && (CurrentTime <= RepeatTime[ButtonIndex]) && (LastMouseOffset.SizeSquared() <= FMath::Square(GMouseDoubleClickArea));
+			if (bHasDoubleClick)
+			{
+				UE_LOG(LogGameInput, Verbose, TEXT("Device %s (DeviceIndex = %d) - %s Double Click"), *UE::GameInput::LexToString(Params.Device), Params.InputDeviceId.GetId(), UE::GameInput::GetMouseButtonName(MouseButton));
+				Params.MessageHandler->OnMouseDoubleClick(nullptr, MouseButton);
+			}
+			else
+			{
+				UE_LOG(LogGameInput, Verbose, TEXT("Device %s (DeviceIndex = %d) - %s Pressed"), *UE::GameInput::LexToString(Params.Device), Params.InputDeviceId.GetId(), UE::GameInput::GetMouseButtonName(MouseButton));
+				Params.MessageHandler->OnMouseDown(nullptr, MouseButton);
+			}
+
+			RepeatTime[ButtonIndex] = CurrentTime + GMouseDoubleClickDelay;
+			LastMouseOffset = FVector2D::ZeroVector;
+		}
+		else
+		{
+			UE_LOG(LogGameInput, Verbose, TEXT("Device %s (DeviceIndex = %d) - %s Released"), *UE::GameInput::LexToString(Params.Device), Params.InputDeviceId.GetId(), UE::GameInput::GetMouseButtonName(MouseButton));
+			Params.MessageHandler->OnMouseUp(MouseButton);
+		}
+	}
+
+	PreviousMouseState = MouseState;
+	LastButtonHeldMask = CurrentButtonHeldMask;
+
+	return true;
+}
+
+void FGameInputMouseDeviceProcessor::ClearState(const FGameInputEventParams& Params)
+{
+	// clear buttons
+	for (int32 ButtonIndex = 0; ButtonIndex < UE_ARRAY_COUNT(UE::GameInput::MouseButtonMappings); ButtonIndex++)
+	{
+		uint32 GameInputButton = UE::GameInput::MouseButtonMappings[ButtonIndex].GameInputButton;
+		if ((GameInputButton & LastButtonHeldMask) != 0)
+		{
+			EMouseButtons::Type MouseButton = UE::GameInput::MouseButtonMappings[ButtonIndex].MouseButton;
+			UE_LOG(LogGameInput, Verbose, TEXT("Device %s (InputDeviceId = %d) - %s Released (via ClearState)"), *UE::GameInput::LexToString(Params.Device), Params.InputDeviceId.GetId(), UE::GameInput::GetMouseButtonName(MouseButton));
+			Params.MessageHandler->OnMouseUp(MouseButton);
+		}
+	}
+
+	// Clear repeat times
+	for (uint32 i = 0; i < MaxSupportedButtons; i++)
+	{
+		RepeatTime[i] = 0.0;
+	}
+
+	// clear previous mouse state
+	LastMouseOffset = FVector2D::ZeroVector;
+	FMemory::Memset(PreviousMouseState, 0);
+}
+
+GameInputKind FGameInputMouseDeviceProcessor::GetSupportedReadingKind() const
+{
+	return GameInputKindMouse;
+}
+
+////////////////////////////////////////////////////////
+// FGameInputTouchDeviceProcessor
+
+bool FGameInputTouchDeviceProcessor::ProcessInput(const FGameInputEventParams& Params)
+{
+	// read the new touch events
+	int32 TouchCount = Params.Reading->GetTouchCount();
+	TArray<GameInputTouchState> InputTouchData;
+	InputTouchData.SetNum(TouchCount, EAllowShrinking::No);
+	Params.Reading->GetTouchState(TouchCount, InputTouchData.GetData());
+
+	// no new touch events and we dont have any previous inputs
+	if (TouchCount == 0 && ActiveTouchPoints == 0)
+	{
+		return false;
+	}
+
+	// We can't tell slate about input messages from an invalid platform user
+	if (!Params.PlatformUserId.IsValid())
+	{
+		UE_LOG(LogGameInput, Verbose, TEXT("[FGameInputTouchDeviceProcessor::EvaluateButtonStates] Attempting to evaluate button states with an invalid platform user id of '%d'. The button messages will not be sent."), Params.PlatformUserId.GetInternalId());
+		return false;
+	}
+
+	auto FindOrAddTouch = [](TArray<FTouchData>& TouchStates, const GameInputTouchState& NewTouchData) -> int32
+	{
+		int32 Index = TouchStates.IndexOfByPredicate([&](const FTouchData& A) {return A.TouchId == NewTouchData.touchId; });
+		if (Index == INDEX_NONE)
+		{
+			// check if we have a free slot
+			Index = TouchStates.IndexOfByPredicate([&](const FTouchData& A) {return A.TouchId == INDEX_NONE; });
+
+			// no free slot - create a new one. (NB. We will never receive more than GameInputDeviceInfo::touchPointCount separate touches)
+			if (Index == INDEX_NONE)
+			{
+				Index = TouchStates.Add(FTouchData());
+			}
+		}
+		return Index;
+	};
+
+	// Reset the active touch state from the previous frame
+	for (FTouchData& TouchData : PreviousTouchData)
+	{
+		TouchData.bIsActive = false;
+	}
+
+	// process all previous frame events and check if there are in list of new events for this frame
+	for (const GameInputTouchState& InputTouch : InputTouchData)
+	{
+		// search for the same touchId in the new events list
+		int32 Index = FindOrAddTouch(PreviousTouchData, InputTouch);
+
+		if (Index == INDEX_NONE)
+		{
+			continue;
+		}
+
+		FTouchData& TouchState = PreviousTouchData[Index];
+
+		// need to scale by resolution
+		FVector2D NewPosition = FVector2D(InputTouch.positionX * (float)MaxTouchX, InputTouch.positionY * (float)MaxTouchY).RoundToVector();
+
+		if (TouchState.TouchId == INDEX_NONE)
+		{
+			Params.MessageHandler->OnTouchStarted(nullptr, NewPosition, InputTouch.pressure, Index, Params.PlatformUserId, Params.InputDeviceId);
+			TouchState.TouchId = InputTouch.touchId;
+			ActiveTouchPoints++;
+		}
+		else if (NewPosition != TouchState.Position)
+		{
+			if (!TouchState.bHasMoved)
+			{
+				Params.MessageHandler->OnTouchFirstMove(NewPosition, InputTouch.pressure, Index, Params.PlatformUserId, Params.InputDeviceId);
+				TouchState.bHasMoved = true;
+			}
+			else
+			{
+				Params.MessageHandler->OnTouchMoved(NewPosition, InputTouch.pressure, Index, Params.PlatformUserId, Params.InputDeviceId);
+			}
+		}
+
+		if (TouchState.Pressure != InputTouch.pressure)
+		{
+			Params.MessageHandler->OnTouchForceChanged(NewPosition, InputTouch.pressure, Index, Params.PlatformUserId, Params.InputDeviceId);
+		}
+
+		TouchState.Pressure = InputTouch.pressure;
+		TouchState.Position = NewPosition;
+		TouchState.bIsActive = true;
+	}
+
+	// process all new event that where not in the stack
+	for (int32 Index = 0; Index < PreviousTouchData.Num(); Index++)
+	{
+		FTouchData& TouchData = PreviousTouchData[Index];
+		if (!TouchData.bIsActive && TouchData.TouchId != INDEX_NONE)
+		{
+			Params.MessageHandler->OnTouchEnded(TouchData.Position, Index, Params.PlatformUserId, Params.InputDeviceId);
+			TouchData = FTouchData();
+			ActiveTouchPoints--;
+		}
+	}
+
+	check(ActiveTouchPoints >= 0);
+
+	return true;
+}
+
+void FGameInputTouchDeviceProcessor::ClearState(const FGameInputEventParams& Params)
+{
+	// There is nothing to be done for touch data clearing...
+}
+
+GameInputKind FGameInputTouchDeviceProcessor::GetSupportedReadingKind() const
+{
+	return GameInputKindTouch;
+}
+
+constexpr float FGameInputRawDeviceProcessor::RawValueToFloatTrigger(const uint8 RawValue) const
+{
+	// Maps the uint8 value of 0-255 to a float between 0.0 and +1.0, like a gamepad trigger.
+	return (static_cast<float>(RawValue) / 255.f);
+}
+
+const float FGameInputRawDeviceProcessor::RawValueToFloatAnalog(uint8 RawValue, const uint8 DeadZone /* = 2 */) const
+{
+	// Apply a simple square deadzone...
+	{
+		// Remember, we are mapping a uint8 (0-255) to a float of -1.0 and +1.0. So, any value
+		// between 0 and 127 is negative and 129-255 is positive. 128 is the center.
+
+		// Calculate the offset of how far this raw value is from center
+		const int32 MaxOffset = FMath::Abs(RawValue - 128);
+
+		// TODO: Implement a better deadzone then this, this is a square one
+		// If we are within the deadzone, set this value to 128 which will translate to 0.0f
+		if (MaxOffset <= DeadZone)
+		{
+			RawValue = 128;
+		}
+	}
+
+	// Maps the uint8 value to a float between -1.0 and +1.0
+	return ((static_cast<float>(RawValue) * (2.f / 255.f)) - 1.f);
+}
+
+const GameInputRawDeviceReportInfo* FGameInputRawDeviceProcessor::ReadCurrentRawInputState(const FGameInputEventParams& Params, IGameInputReading* ReadingToUse)
+{
+	if (!ReadingToUse)
+	{
+		UE_LOG(LogGameInput, Error, TEXT("[ReadCurrentRawInputState] Cannot read raw input state, ReadingToUse was null (Device %s) "), *UE::GameInput::LexToString(Params.Device));
+		return nullptr;
+	}
+
+	TComPtr<IGameInputRawDeviceReport> RawReport;
+	const bool bSuccessfulReading = ReadingToUse->GetRawReport(&RawReport);
+	if (!bSuccessfulReading)
+	{
+		UE_LOG(LogGameInput, Error, TEXT("[ReadCurrentRawInputState] Unsuccessful reading of raw input report! (GetRawReport failed) (Device %s) "), *UE::GameInput::LexToString(Params.Device));
+		return nullptr;
+	}
+
+	const int32 RawRepDataSize = RawReport->GetRawDataSize();
+
+	// Ensure that the current data array is populated to the size needed.
+	CurrentRawData.Reset();
+	CurrentRawData.AddZeroed(RawRepDataSize);
+
+	// Ensure that the previous data array is populated to the size needed so that we can compare its values per-index
+	if (PreviousRawData.Num() < RawRepDataSize)
+	{
+		PreviousRawData.AddZeroed(RawRepDataSize - PreviousRawData.Num());
+	}
+
+	// This will populate the values in the CurrentRawData so that we can process them
+	const int32 NumReadBytes = RawReport->GetRawData(RawReport->GetRawDataSize(), CurrentRawData.GetData());
+
+	const GameInputRawDeviceReportInfo* RawReportInfo = RawReport->GetReportInfo();
+	if (!RawReportInfo)
+	{
+		UE_LOG(LogGameInput, Warning, TEXT("[ProcessRawReport] Unsuccessful reading of raw input report! (GameInputRawDeviceReportInfo is null) (Device %s) "), *UE::GameInput::LexToString(Params.Device));
+	}
+
+	return RawReportInfo;
+}
+
+bool FGameInputRawDeviceProcessor::ProcessInput(const FGameInputEventParams& Params)
+{
+	bool bRes = false;
+
+	// Can't do anything for an invalid platform user
+	if (!Params.PlatformUserId.IsValid())
+	{
+		return bRes;
+	}
+
+	const FGameInputDeviceConfiguration* DeviceConfig = GetDefault<UGameInputDeveloperSettings>()->FindDeviceConfiguration(Params.GetDeviceInfo());
+
+	// Check that we have a valid config before bothering to read the raw report
+	if (!DeviceConfig)
+	{
+		UE_LOG(LogGameInput, Verbose, TEXT("[ProcessRawReport] (Device %s) Does not have a valid FGameInputDeviceConfiguration in the UGameInputDeveloperSettings. We can't process Raw Input without it. Exiting."), *UE::GameInput::LexToString(Params.Device));
+		return bRes;
+	}
+
+	// Skip if this device config isn't even supposed to be processing raw input values
+	if (!DeviceConfig->bProcessRawReportData)
+	{
+		return bRes;
+	}
+
+	// Read from the current reading the CURRENT frame
+	const GameInputRawDeviceReportInfo* RawReportInfo = ReadCurrentRawInputState(Params, Params.Reading);
+
+	// Actually process the current raw input values if this reading ID matches the one we want
+	if (RawReportInfo && RawReportInfo->id == DeviceConfig->RawReportReadingId)
+	{
+		// Process the BUTTON types here, we only want to process analog events once per frame to avoid over-accumulation
+		constexpr bool bShouldProcessButtons = true;
+		constexpr bool bShouldProcessAnalog = false;
+		bRes |= ProcessAllRawValues(Params, DeviceConfig, bShouldProcessButtons, bShouldProcessAnalog);
+
+		++NumReadingsProcessedThisFrame;
+	}
+
+	return bRes;
+}
+
+
+bool FGameInputRawDeviceProcessor::PostProcessInput(const FGameInputEventParams& Params)
+{
+	bool bRes = false;
+
+	const bool bProcessedAnyGamepadButtonsThisFrame = (NumReadingsProcessedThisFrame > 0);
+
+	// This is the last reading of this frame, reset the counter to 0
+	NumReadingsProcessedThisFrame = 0;
+
+	// Can't do anything for an invalid platform user
+	if (!Params.PlatformUserId.IsValid())
+	{
+		return bRes;
+	}
+
+	const FGameInputDeviceConfiguration* DeviceConfig = GetDefault<UGameInputDeveloperSettings>()->FindDeviceConfiguration(Params.GetDeviceInfo());
+
+	// Check that we have a valid config before bothering to read the raw report
+	if (!DeviceConfig)
+	{
+		UE_LOG(LogGameInput, Verbose, TEXT("[ProcessRawReport] (Device %s) Does not have a valid FGameInputDeviceConfiguration in the UGameInputDeveloperSettings. We can't process Raw Input without it. Exiting."), *UE::GameInput::LexToString(Params.Device));
+		return bRes;
+	}
+
+	// On the last input reading for the frame, the "Current Reading" should always be null. we only care for the LastReading here
+	ensure(Params.Reading == nullptr);
+
+	if (!Params.PreviousReading)
+	{
+		return bRes;
+	}
+
+	// Read from the current reading the PREVIOUS frame
+	const GameInputRawDeviceReportInfo* RawReportInfo = ReadCurrentRawInputState(Params, Params.PreviousReading);
+
+	// Actually process the current raw input values if this reading ID matches the one we want
+	if (RawReportInfo && RawReportInfo->id == DeviceConfig->RawReportReadingId)
+	{
+		// We always want to process analog events on the last frame
+		constexpr bool bShouldPorcessAnalog = true;
+		// We only need to process buttons this frame too if there have been no other readings yet.
+		const bool bNeedToProcessButtons = !bProcessedAnyGamepadButtonsThisFrame;
+
+		bRes |= ProcessAllRawValues(Params, DeviceConfig, bNeedToProcessButtons, bShouldPorcessAnalog);
+	}
+
+	// Track our previous input only on the last input frame. We don't want any duplicate readings
+	PreviousRawData = CurrentRawData;
+
+	return bRes;
+}
+
+bool FGameInputRawDeviceProcessor::ProcessAllRawValues(const FGameInputEventParams& Params, const FGameInputDeviceConfiguration* DeviceConfig, const bool bShouldProcessButtons, const bool bShouldProcessAnalog)
+{
+	ensure(CurrentRawData.Num() == PreviousRawData.Num());
+
+	ensure(Params.PlatformUserId.IsValid());
+
+	bool bRes = false;
+
+	for (int32 i = 0; i < CurrentRawData.Num(); ++i)
+	{
+		const uint8 Val = CurrentRawData[i];
+
+		if (const FGameInputRawDeviceReportData* AxisData = DeviceConfig->RawReportMappingData.Find(i))
+		{
+			if (AxisData->TranslationBehavior == ERawDeviceReportTranslationBehavior::TreatAsButtonBitmask)
+			{
+				if (bShouldProcessButtons)
+				{
+					bRes |= ProcessRawInputValueAsBitmask(Params, i, AxisData);
+				}				
+			}
+			else if (AxisData->TranslationBehavior == ERawDeviceReportTranslationBehavior::TreatAsPackedAxisPair)
+			{
+				// Only call this function on the higher index
+				if (bShouldProcessAnalog && i == AxisData->HigherBitAxisIndex)
+				{
+					ProcessRawInputValueAsAanalogPaired(Params, AxisData);
+				}
+			}
+			// All other methods require a valid key name on the config
+			else if (AxisData->KeyName.IsValid())
+			{
+				// Treat this value as a button. If it is non-zero then consider it pressed. If it is zero, then it is not pressed.
+				if (AxisData->TranslationBehavior == ERawDeviceReportTranslationBehavior::TreatAsButton)
+				{
+					if (bShouldProcessButtons)
+					{
+						bRes |= ProcessRawInputValueAsButton(Params, i, AxisData);
+					}					
+				}
+				// Otherwise we can do analog values, which can be either "trigger" or "analog" types
+				else if (bShouldProcessAnalog)
+				{
+					bRes |= ProcessRawInputValueAsAanalog(Params, i, AxisData);
+				}
+			}
+			else
+			{
+				// You want a valid key name here, throw a warning if your config is wrong
+				UE_LOG(LogGameInput, Warning, TEXT("[ProcessRawReport] Invalid key name for raw report axis at index %d with value of %u (Device %s)"), i, Val, *UE::GameInput::LexToString(Params.Device));
+			}
+		}
+		else if (Val > 0)
+		{
+			UE_LOG(LogGameInput, VeryVerbose, TEXT("[ProcessRawReport] (Device %s) No raw device report config for axis '%d' with value of %u"), *UE::GameInput::LexToString(Params.Device), i, Val);
+		}
+	}
+
+	return bRes;
+}
+
+bool FGameInputRawDeviceProcessor::ProcessRawInputValueAsBitmask(const FGameInputEventParams& Params, const int32 RawValueIndex, const FGameInputRawDeviceReportData* AxisData)
+{
+	check(AxisData && AxisData->TranslationBehavior == ERawDeviceReportTranslationBehavior::TreatAsButtonBitmask);
+
+	const uint8 Val = CurrentRawData[RawValueIndex];
+	const uint8 PrevVal = PreviousRawData[RawValueIndex];
+
+	// Ensure we have a compatible key map setup
+	FPerRawInputIndexData& IndexData = RawInputIndexDataMap.FindOrAdd(RawValueIndex);
+
+	if (IndexData.KeyNameMap.IsEmpty())
+	{
+		// The settings use their button map as a bit number, so the key map we actually need to use
+		// is 1 << that bit to have EvaluateButtonStates work correctly
+		for (const TPair<int32, FName>& MappingPair : AxisData->ButtonBitMaskMappings)
+		{
+			IndexData.KeyNameMap.Add({ 1 << MappingPair.Key, MappingPair.Value });
+		}
+	}
+
+	const uint32 CurrentValue32 = static_cast<uint32>(Val);
+	uint32 PreviousValue32 = static_cast<uint32>(PrevVal);
+
+	// Note: Set the max supported buttons here to 1 because we only care about the first bit
+	EvaluateButtonStates(
+		Params,
+		CurrentValue32,
+		OUT PreviousValue32,
+		IndexData.RepeatTime,
+		IndexData.KeyNameMap,
+		FPerRawInputIndexData::MaxSupportedButtons);
+
+	// If this value is non-zero then we had a reading
+	return Val != 0;
+}
+
+bool FGameInputRawDeviceProcessor::ProcessRawInputValueAsButton(const FGameInputEventParams& Params, const int32 RawValueIndex, const FGameInputRawDeviceReportData* AxisData)
+{
+	check(AxisData && AxisData->TranslationBehavior == ERawDeviceReportTranslationBehavior::TreatAsButton);
+
+	const uint8 Val = CurrentRawData[RawValueIndex];
+	const uint8 PrevVal = PreviousRawData[RawValueIndex];
+
+	// Just treat this button as pressed when non-zero, and not pressed when 0.
+	const uint32 CurrentValue32 = FMath::Clamp<uint32>(static_cast<uint32>(Val), 0, 1);
+	uint32 PreviousValue32 = FMath::Clamp<uint32>(static_cast<uint32>(PrevVal), 0, 1);
+
+	// Map the value of the key name to this 1 if we haven't already
+	FPerRawInputIndexData& IndexData = RawInputIndexDataMap.FindOrAdd(RawValueIndex);
+	if (IndexData.KeyNameMap.IsEmpty())
+	{
+		IndexData.KeyNameMap.Add(1, AxisData->KeyName);
+	}
+
+	// Note: Set the max supported buttons here to 1 because we only care about the first bit
+	EvaluateButtonStates(
+		Params,
+		CurrentValue32,
+		OUT PreviousValue32,
+		IndexData.RepeatTime,
+		IndexData.KeyNameMap,
+		/* maxSupportedButtons */ 1);
+
+	// If this value is non-zero then we had a reading
+	return Val != 0;
+}
+
+bool FGameInputRawDeviceProcessor::ProcessRawInputValueAsAanalog(const FGameInputEventParams& Params, const int32 RawValueIndex, const FGameInputRawDeviceReportData* AxisData)
+{
+	check(AxisData);
+
+	const uint8 Val = CurrentRawData[RawValueIndex];
+	const uint8 PrevVal = PreviousRawData[RawValueIndex];
+
+	const float CurrentValueFloat = AxisData->Scalar * (AxisData->TranslationBehavior == ERawDeviceReportTranslationBehavior::TreatAsTrigger ? RawValueToFloatTrigger(Val) : RawValueToFloatAnalog(Val, AxisData->AnalogDeadzone));
+	const float PreviousValueFloat = AxisData->Scalar * (AxisData->TranslationBehavior == ERawDeviceReportTranslationBehavior::TreatAsTrigger ? RawValueToFloatTrigger(PrevVal) : RawValueToFloatAnalog(PrevVal, AxisData->AnalogDeadzone));
+
+	if (!AxisData->KeyName.IsValid())
+	{
+		return false;
+	}
+
+	OnControllerAnalog(Params, AxisData->KeyName, CurrentValueFloat, PreviousValueFloat, UE::GameInput::GamepadLeftStickDeadzone);
+
+	// We had a reading as long as it is non-zero
+	return CurrentValueFloat != 0.0f;
+}
+
+bool FGameInputRawDeviceProcessor::ProcessRawInputValueAsAanalogPaired(const FGameInputEventParams& Params, const FGameInputRawDeviceReportData* AxisData)
+{
+	check(AxisData);
+
+	if (!CurrentRawData.IsValidIndex(AxisData->LowerBitAxisIndex) || !CurrentRawData.IsValidIndex(AxisData->HigherBitAxisIndex))
+	{
+		return false;
+	}
+
+	// Get the current value
+	const uint8 CurrentLowerVal = CurrentRawData[AxisData->LowerBitAxisIndex];
+	const uint8 CurrentHigherVal = CurrentRawData[AxisData->HigherBitAxisIndex];
+	
+	// Combine the two values into a single int16. Do this by 
+	// shifting the higher value up by 8 bits, and then just use the lower value in our int16	
+	const int16 CurrentPackedVal = (int16)((CurrentHigherVal << 8) | CurrentLowerVal);
+	const float CurrentValueFloat = (static_cast<float>(CurrentPackedVal) / 32767.f);
+
+	// Get the previous value
+	const uint8 PreviousLowerVal = PreviousRawData[AxisData->LowerBitAxisIndex];
+	const uint8 PreviousHigherVal = PreviousRawData[AxisData->HigherBitAxisIndex];
+
+	const int16 PreviousPackedVal = (int16)((PreviousHigherVal << 8) | PreviousLowerVal);
+	const float PreviousValueFloat = (static_cast<float>(PreviousPackedVal) / 32767.f);
+
+	if (!AxisData->KeyName.IsValid())
+	{
+		return false;
+	}
+
+	OnControllerAnalog(Params, AxisData->KeyName, CurrentValueFloat, PreviousValueFloat, UE::GameInput::GamepadLeftStickDeadzone);
+
+	// We had a reading as long as it is non-zero
+	return CurrentValueFloat != 0.0f;
+}
+
+void FGameInputRawDeviceProcessor::ClearState(const FGameInputEventParams& Params)
+{
+	// Can't do anything for an invalid platform user
+	if (!Params.PlatformUserId.IsValid())
+	{
+		return;
+	}
+
+	const FGameInputDeviceConfiguration* DeviceConfig = GetDefault<UGameInputDeveloperSettings>()->FindDeviceConfiguration(Params.GetDeviceInfo());
+
+	// Check that we have a valid config before bothering to read the raw report
+	if (!DeviceConfig)
+	{
+		UE_LOG(LogGameInput, Verbose, TEXT("[ClearStateRawReport] No have a valid FGameInputDeviceConfiguration in the UGameInputDeveloperSettings. We can't process Raw Input without it. Exiting. (Device %s)"), *UE::GameInput::LexToString(Params.Device));
+		return;
+	}
+
+	if (!DeviceConfig->bProcessRawReportData)
+	{
+		return;
+	}
+
+	// Reset our current values to 0...
+	for (int32 i = 0; i < CurrentRawData.Num(); ++i)
+	{
+		CurrentRawData[i] = 0;
+	}
+
+	// ... and then process the raw values as if there is 0 input. We want to process all types when clearing.
+	constexpr bool bShouldPorcessButtons = true;
+	constexpr bool bShouldPorcessAnalog = true;
+
+	ProcessAllRawValues(Params, DeviceConfig, bShouldPorcessButtons, bShouldPorcessAnalog);
+}
+
+GameInputKind FGameInputRawDeviceProcessor::GetSupportedReadingKind() const
+{
+	return GameInputKindRawDeviceReport;
+}
+
+#endif	// GAME_INPUT_SUPPORT
\ No newline at end of file
diff --git a/Engine/Plugins/Runtime/GameInput/Source/GameInputBase/Private/GameInputKeyTypes.cpp b/Engine/Plugins/Runtime/GameInput/Source/GameInputBase/Private/GameInputKeyTypes.cpp
new file mode 100644
--- /dev/null
+++ b/Engine/Plugins/Runtime/GameInput/Source/GameInputBase/Private/GameInputKeyTypes.cpp
@@ -0,0 +1,60 @@
+// Copyright Epic Games, Inc. All Rights Reserved.
+
+#include "GameInputKeyTypes.h"
+
+//        
+// Racing Wheel
+//        
+
+// Analog types
+const FKey FGameInputKeys::RacingWheel_Brake("GameInput_RacingWheel_Brake");
+const FKey FGameInputKeys::RacingWheel_Clutch("GameInput_RacingWheel_Clutch");
+const FKey FGameInputKeys::RacingWheel_Handbrake("GameInput_RacingWheel_Handbrake");
+const FKey FGameInputKeys::RacingWheel_Throttle("GameInput_RacingWheel_Throttle");
+const FKey FGameInputKeys::RacingWheel_Wheel("GameInput_RacingWheel_Wheel");
+const FKey FGameInputKeys::RacingWheel_PatternShifterGear("GameInput_RacingWheel_PatternShifterGear");
+
+// Button types
+const FKey FGameInputKeys::RacingWheel_None("GameInput_RacingWheel_None");
+const FKey FGameInputKeys::RacingWheel_Menu("GameInput_RacingWheel_Menu");
+const FKey FGameInputKeys::RacingWheel_View("GameInput_RacingWheel_View");
+const FKey FGameInputKeys::RacingWheel_PreviousGear("GameInput_RacingWheel_PreviousGear");
+const FKey FGameInputKeys::RacingWheel_NextGear("GameInput_RacingWheel_NextGear");
+
+//
+// Flight Stick
+//
+
+// Analog Types
+const FKey FGameInputKeys::FlightStick_Roll("FlightStick_Roll");
+const FKey FGameInputKeys::FlightStick_Pitch("FlightStick_Pitch");
+const FKey FGameInputKeys::FlightStick_Yaw("FlightStick_Yaw");
+const FKey FGameInputKeys::FlightStick_Throttle("FlightStick_Throttle");
+
+// Button types
+const FKey FGameInputKeys::FlightStick_None("GameInput_FlightStick_None");
+const FKey FGameInputKeys::FlightStick_Menu("GameInput_FlightStick_Menu");
+const FKey FGameInputKeys::FlightStick_View("GameInput_FlightStick_View");
+const FKey FGameInputKeys::FlightStick_FirePrimary("GameInput_FlightStick_FirePrimary");
+const FKey FGameInputKeys::FlightStick_FireSecondary("GameInput_FlightStick_FireSecondary");
+
+//
+// Arcade Stick
+//
+
+// Button Types
+const FKey FGameInputKeys::ArcadeStick_None("GameInput_ArcadeStick_None");
+const FKey FGameInputKeys::ArcadeStick_Menu("GameInput_ArcadeStick_Menu");
+const FKey FGameInputKeys::ArcadeStick_View("GameInput_ArcadeStick_View");
+const FKey FGameInputKeys::ArcadeStick_Up("GameInput_ArcadeStick_Up");
+const FKey FGameInputKeys::ArcadeStick_Down("GameInput_ArcadeStick_Down");
+const FKey FGameInputKeys::ArcadeStick_Left("GameInput_ArcadeStick_Left");
+const FKey FGameInputKeys::ArcadeStick_Right("GameInput_ArcadeStick_Right");
+const FKey FGameInputKeys::ArcadeStick_Action1("GameInput_ArcadeStick_Action1");
+const FKey FGameInputKeys::ArcadeStick_Action2("GameInput_ArcadeStick_Action2");
+const FKey FGameInputKeys::ArcadeStick_Action3("GameInput_ArcadeStick_Action3");
+const FKey FGameInputKeys::ArcadeStick_Action4("GameInput_ArcadeStick_Action4");
+const FKey FGameInputKeys::ArcadeStick_Action5("GameInput_ArcadeStick_Action5");
+const FKey FGameInputKeys::ArcadeStick_Action6("GameInput_ArcadeStick_Action6");
+const FKey FGameInputKeys::ArcadeStick_Special1("GameInput_ArcadeStick_Special1");
+const FKey FGameInputKeys::ArcadeStick_Special2("GameInput_ArcadeStick_Special2");
\ No newline at end of file
diff --git a/Engine/Plugins/Runtime/GameInput/Source/GameInputBase/Private/GameInputLogging.cpp b/Engine/Plugins/Runtime/GameInput/Source/GameInputBase/Private/GameInputLogging.cpp
new file mode 100644
--- /dev/null
+++ b/Engine/Plugins/Runtime/GameInput/Source/GameInputBase/Private/GameInputLogging.cpp
@@ -0,0 +1,5 @@
+// Copyright Epic Games, Inc. All Rights Reserved.
+
+#include "GameInputLogging.h"
+
+DEFINE_LOG_CATEGORY(LogGameInput);
\ No newline at end of file
diff --git a/Engine/Plugins/Runtime/GameInput/Source/GameInputBase/Private/GameInputUtils.cpp b/Engine/Plugins/Runtime/GameInput/Source/GameInputBase/Private/GameInputUtils.cpp
new file mode 100644
--- /dev/null
+++ b/Engine/Plugins/Runtime/GameInput/Source/GameInputBase/Private/GameInputUtils.cpp
@@ -0,0 +1,178 @@
+// Copyright Epic Games, Inc. All Rights Reserved.
+
+#include "GameInputUtils.h"
+
+#if GAME_INPUT_SUPPORT
+
+namespace UE::GameInput
+{
+	FString LexToString(IGameInputDevice* Device)
+	{
+		FString Result = TEXT("<invalid device>");
+
+		if (Device)
+		{
+			const GameInputDeviceInfo* Info = Device->GetDeviceInfo();
+			const bool bIsVirtualDevice = Info->deviceFamily == GameInputFamilyVirtual;
+			
+			Result = FString::Printf(TEXT("%llp   DeviceName: %s %s ProdId: %04x VendId: %04x"), 
+				Device, 
+				Info->displayName, 
+				bIsVirtualDevice ? TEXT("(virtual)") : TEXT(""),
+				Info->productId,
+				Info->vendorId);
+		}
+
+		return Result;
+	}
+	
+	FString LexToString(GameInputDeviceStatus DeviceStatus)
+	{
+		if (DeviceStatus == GameInputDeviceNoStatus)
+		{
+			return TEXT("NoStatus");
+		}
+		else if (DeviceStatus == GameInputDeviceAnyStatus)
+		{
+			return TEXT("Any");
+		}
+
+		FString Result = TEXT("");
+#define DEVICE_STATUS(StatusFlag,DisplayName) if( DeviceStatus & StatusFlag ) Result += (FString(DisplayName) + TEXT("|"));
+		DEVICE_STATUS(GameInputDeviceConnected, TEXT("Connected"));
+		DEVICE_STATUS(GameInputDeviceInputEnabled, TEXT("InputEnabled"));
+		DEVICE_STATUS(GameInputDeviceOutputEnabled, TEXT("OutputEnabled"));
+		DEVICE_STATUS(GameInputDeviceRawIoEnabled, TEXT("RawIoEnabled"));
+		DEVICE_STATUS(GameInputDeviceAudioCapture, TEXT("AudioCapture"));
+		DEVICE_STATUS(GameInputDeviceAudioRender, TEXT("AudioRender"));
+		DEVICE_STATUS(GameInputDeviceSynchronized, TEXT("Synchronized"));
+		DEVICE_STATUS(GameInputDeviceWireless, TEXT("Wireless"));
+		DEVICE_STATUS(GameInputDeviceUserIdle, TEXT("UserIdle"));
+#undef DEVICE_STATUS
+
+		Result.RemoveFromEnd(TEXT("|"));
+		return Result;
+	}
+	
+	FString LexToString(GameInputKind InputKind)
+	{
+		if (InputKind == GameInputKindUnknown)
+		{
+			return TEXT("Unknown");
+		}
+
+		FString Result = TEXT("");
+#define INPUT_KIND_STRING(StatusFlag, DisplayName) if( InputKind & StatusFlag ) Result += (FString(DisplayName) + TEXT("|"));
+		INPUT_KIND_STRING(GameInputKindRawDeviceReport, TEXT("RawDeviceReport"));
+		INPUT_KIND_STRING(GameInputKindControllerAxis, TEXT("ControllerAxis"));
+		INPUT_KIND_STRING(GameInputKindControllerButton, TEXT("ControllerButton"));
+		INPUT_KIND_STRING(GameInputKindControllerSwitch, TEXT("ControllerSwitch"));
+		INPUT_KIND_STRING(GameInputKindController, TEXT("Controller"));
+		INPUT_KIND_STRING(GameInputKindKeyboard, TEXT("Keyboard"));
+		INPUT_KIND_STRING(GameInputKindMouse, TEXT("Mouse"));
+		INPUT_KIND_STRING(GameInputKindTouch, TEXT("Touch"));
+		INPUT_KIND_STRING(GameInputKindMotion, TEXT("Motion"));
+		INPUT_KIND_STRING(GameInputKindArcadeStick, TEXT("ArcadeStick"));
+		INPUT_KIND_STRING(GameInputKindFlightStick, TEXT("FlightStick"));
+		INPUT_KIND_STRING(GameInputKindGamepad, TEXT("Gamepad"));
+		INPUT_KIND_STRING(GameInputKindRacingWheel, TEXT("RacingWheel"));
+		INPUT_KIND_STRING(GameInputKindUiNavigation, TEXT("UiNavigation"));
+#undef INPUT_KIND_STRING
+
+		Result.RemoveFromEnd(TEXT("|"));
+		return Result;
+	}
+
+	FString LexToString(GameInputSwitchPosition SwitchPos)
+	{
+		switch(SwitchPos)
+		{
+		case GameInputSwitchCenter: return TEXT("SwitchCenter");
+		case GameInputSwitchUp: return TEXT("SwitchUp");
+		case GameInputSwitchUpRight: return TEXT("SwitchUpRight");
+		case GameInputSwitchRight: return TEXT("SwitchRight");
+		case GameInputSwitchDownRight: return TEXT("SwitchDownRight");
+		case GameInputSwitchDown: return TEXT("SwitchDown");
+		case GameInputSwitchDownLeft: return TEXT("SwitchDownLeft");
+		case GameInputSwitchLeft: return TEXT("SwitchLeft");
+		case GameInputSwitchUpLeft: return TEXT("SwitchUpLeft");
+		}
+
+		return TEXT("Unknown Switch Position");
+	}
+
+	FString LexToString(const APP_LOCAL_DEVICE_ID& DeviceId)
+	{
+		return BytesToHex(DeviceId.value, APP_LOCAL_DEVICE_ID_SIZE);
+	}
+	
+	EInputDeviceConnectionState DeviceStateToConnectionState(GameInputDeviceStatus InCurrentStatus, GameInputDeviceStatus InPreviousStatus)
+	{
+		InCurrentStatus = InCurrentStatus & GameInputDeviceConnected;
+		InPreviousStatus = InPreviousStatus & GameInputDeviceConnected;
+
+		// If the current status has changed to connected, or the previous and current status are connected, then respond as such
+		// The current and previous state may both respond as connected upon Application boot
+		if ((InCurrentStatus > InPreviousStatus) || (InCurrentStatus && InPreviousStatus))
+		{
+			return EInputDeviceConnectionState::Connected;
+		}
+		else if (InPreviousStatus > InCurrentStatus)
+		{
+			return EInputDeviceConnectionState::Disconnected;
+		}
+		else
+		{
+			// Returned in the case of there being no connection change
+			return EInputDeviceConnectionState::Invalid;
+		}
+	}
+
+	const TCHAR* GetMouseButtonName(EMouseButtons::Type MouseButton)
+	{
+		switch (MouseButton)
+		{
+		case EMouseButtons::Left: return TEXT("Left");
+		case EMouseButtons::Right: return TEXT("Right");
+		case EMouseButtons::Middle: return TEXT("Middle");
+		case EMouseButtons::Thumb01: return TEXT("Thumb01");
+		case EMouseButtons::Thumb02: return TEXT("Thumb02");
+		}
+		return TEXT("Invalid");
+	}
+
+	bool GameInputButtonToUnrealName(const TMap<uint32, FGamepadKeyNames::Type>& UEButtonMap, uint32 ButtonMask, OUT FGamepadKeyNames::Type& OutKeyName)
+	{
+		if (const FGamepadKeyNames::Type* Name = UEButtonMap.Find(ButtonMask))
+		{
+			OutKeyName = *Name;
+			return true;
+		}
+
+		return false;
+	}
+
+	const TMap<GameInputSwitchPosition, TArray<FGamepadKeyNames::Type>>& GetSwitchButtonMap()
+	{
+		static const TMap<GameInputSwitchPosition, TArray<FGamepadKeyNames::Type>> SwitchButtonMap
+		{
+			{ GameInputSwitchCenter,			{ } },
+			{ GameInputSwitchUp,				{ FGamepadKeyNames::DPadUp } },
+			{ GameInputSwitchUpRight,			{ FGamepadKeyNames::DPadUp, FGamepadKeyNames::DPadRight } },
+			{ GameInputSwitchRight,				{ FGamepadKeyNames::DPadRight } },
+			{ GameInputSwitchDownRight,			{ FGamepadKeyNames::DPadDown, FGamepadKeyNames::DPadRight } },
+			{ GameInputSwitchDown,				{ FGamepadKeyNames::DPadDown } },
+			{ GameInputSwitchDownLeft,			{ FGamepadKeyNames::DPadDown, FGamepadKeyNames::DPadLeft } },
+			{ GameInputSwitchLeft,				{ FGamepadKeyNames::DPadLeft } },
+			{ GameInputSwitchUpLeft,			{ FGamepadKeyNames::DPadUp, FGamepadKeyNames::DPadLeft } },
+		};
+		return SwitchButtonMap;
+	}
+
+	const TArray<FGamepadKeyNames::Type>* SwitchPositionToUnrealName(const GameInputSwitchPosition ButtonMask)
+	{
+		return GetSwitchButtonMap().Find(ButtonMask);
+	}
+}
+
+#endif	// #if GAME_INPUT_SUPPORT
\ No newline at end of file
diff --git a/Engine/Plugins/Runtime/GameInput/Source/GameInputBase/Private/IGameInputDeviceInterface.cpp b/Engine/Plugins/Runtime/GameInput/Source/GameInputBase/Private/IGameInputDeviceInterface.cpp
new file mode 100644
--- /dev/null
+++ b/Engine/Plugins/Runtime/GameInput/Source/GameInputBase/Private/IGameInputDeviceInterface.cpp
@@ -0,0 +1,473 @@
+// Copyright Epic Games, Inc. All Rights Reserved.
+
+#include "IGameInputDeviceInterface.h"
+#include "GameInputDeveloperSettings.h"
+#include "GameInputLogging.h"
+#include "GameInputUtils.h"
+#include "GenericPlatform/GenericPlatformInputDeviceMapper.h"
+#include "Misc/CoreDelegates.h"
+#include "GenericPlatform/IInputInterface.h"
+#include "HAL/IConsoleManager.h"
+
+#if GAME_INPUT_SUPPORT
+namespace UE::GameInput
+{
+	static TAutoConsoleVariable<bool> CVarGameInputEnumerateDeviceTypeOnConnection
+	(
+		TEXT("gameinput.bEnumerateDeviceTypeOnConnection"),
+		true,
+		TEXT("If true, EnumerateCurrentlyConnectedDeviceTypes will be called upon device connect and disconnect"),
+		ECVF_Default
+	);
+};
+#endif	// GAME_INPUT_SUPPORT
+
+IGameInputDeviceInterface::IGameInputDeviceInterface(const TSharedRef<FGenericApplicationMessageHandler>& InMessageHandler, IGameInput* InGameInput)
+#if GAME_INPUT_SUPPORT
+	: MessageHandler(InMessageHandler)
+	, GameInput(InGameInput)
+	, CurrentlyConnectedDeviceTypes(GameInputKindUnknown)
+	, bWasinitialized(false)
+	, bIsAppCurrentlyConstrained(false)
+	, bWasAppConstrainedLastTick(false)
+#endif	// GAME_INPUT_SUPPORT
+{
+}
+
+IGameInputDeviceInterface::~IGameInputDeviceInterface()
+{
+#if GAME_INPUT_SUPPORT
+	if (ConnectionChangeCallbackToken)
+	{
+		if (GameInput)
+		{
+			GameInput->UnregisterCallback(ConnectionChangeCallbackToken, UINT64_MAX);
+			GameInput.Reset();
+		}
+
+		ConnectionChangeCallbackToken = GAMEINPUT_INVALID_CALLBACK_TOKEN_VALUE;
+	}
+
+	FCoreDelegates::ApplicationWillDeactivateDelegate.RemoveAll(this);
+	FCoreDelegates::ApplicationHasReactivatedDelegate.RemoveAll(this);
+#endif	// GAME_INPUT_SUPPORT
+}
+
+void IGameInputDeviceInterface::Initialize()
+{
+#if GAME_INPUT_SUPPORT
+	BindDeviceStatusCallbacks();
+	bWasinitialized = true;
+#endif	// GAME_INPUT_SUPPORT
+}
+
+void IGameInputDeviceInterface::Tick(float DeltaTime)
+{
+	// Required by IInputDevice
+}
+
+void IGameInputDeviceInterface::SendControllerEvents()
+{
+#if GAME_INPUT_SUPPORT
+
+	ensureAlwaysMsgf(bWasinitialized, TEXT("A Game Input Device was not initalized. You must call IGameInputDeviceInterface::Initialize after creating a new device!"));
+
+	FScopeLock Lock(&DeviceInfoCS);
+
+	// At this point we should always have a valid Game Input object
+	check(GameInput);
+
+	// On the first update coming back from being constrained, we want to reset
+	// the state of inputs
+	if (bIsAppCurrentlyConstrained && !bWasAppConstrainedLastTick)
+	{
+		DetermineStateAfterFirstUnconstrainedUpdate();
+	}
+
+	// Process any input devices so long as we are not constrained
+	if (!bIsAppCurrentlyConstrained)
+	{
+		// Handle any device connection/disconnection state changes first before attempting to process any devices
+		ProcessDeferredDeviceConnectionChanges();
+		
+		// A map of Platform users to a bitmask of any reading kinds that were processed this frame.
+		// This is used by the devices to keep track of which users had which readings, and use that
+		// state to determine if we can process a given reading.
+		static TMap<FPlatformUserId, GameInputKind> PlatformUsersWhoHaveHadInputThisFrame;
+		PlatformUsersWhoHaveHadInputThisFrame.Reset();
+
+		// The allowed input kinds that we can read from the IGameInput interface.
+		const GameInputKind AllowedInputKindsThisFrame = GetCurrentGameInputKindSupport();
+
+		for (const TSharedPtr<FGameInputDeviceContainer>& KnownDevice : DeviceData)
+		{
+			const FPlatformUserId UserId = KnownDevice->GetPlatformUserId();
+
+			// Figure out what GameInputKind's this platform user has already handled this frame.
+			GameInputKind& AlreadyProcessedInputKinds = PlatformUsersWhoHaveHadInputThisFrame.FindOrAdd(UserId, /* default init value */ GameInputKindUnknown);
+
+			// Actually process our input here, and get a bitmask back of what GameInputKind's have sent events.
+			const GameInputKind InputKindsWithEvents = KnownDevice->ProcessInput(GameInput, AllowedInputKindsThisFrame, AlreadyProcessedInputKinds);
+
+			// Keep track what kinds of input that this platform user has processed this frame for use on the next device
+			AlreadyProcessedInputKinds |= InputKindsWithEvents;
+			
+			// Keep track of the most recent device that is being used by each given platform user
+			// We need to do this whenever we receive input, which is true as long as there was an input kind processed this frame.
+			if (InputKindsWithEvents != GameInputKindUnknown)
+			{
+				// If we know input came from a device associated to this platform user already, then check our timestamp to see if it is newer then it
+				if (const TSharedPtr<FGameInputDeviceContainer>* MostRecentKnownDevice = PlatformUserIdToMostRecentDeviceContainer.Find(UserId))
+				{
+					// If the current device that was just processed has a more recent timestamp then the known one, then
+					// use it as our most recent device instead
+					if (KnownDevice->GetLastReadingTimestamp() > (*MostRecentKnownDevice)->GetLastReadingTimestamp())
+					{
+						PlatformUserIdToMostRecentDeviceContainer[UserId] = KnownDevice;
+					}
+				}
+				else
+				{
+					PlatformUserIdToMostRecentDeviceContainer.Add(UserId, KnownDevice);
+				}
+			}
+		}
+	}
+
+	bWasAppConstrainedLastTick = bIsAppCurrentlyConstrained;
+
+#endif// GAME_INPUT_SUPPORT
+}
+
+void IGameInputDeviceInterface::SetMessageHandler(const TSharedRef<FGenericApplicationMessageHandler>& InMessageHandler)
+{
+#if GAME_INPUT_SUPPORT
+	MessageHandler = InMessageHandler;
+
+	for (const TSharedPtr<FGameInputDeviceContainer>& KnownDevice : DeviceData)
+	{
+		KnownDevice->SetMessageHandler(InMessageHandler);
+	}
+#endif	// GAME_INPUT_SUPPORT
+}
+
+bool IGameInputDeviceInterface::Exec(UWorld* InWorld, const TCHAR* Cmd, FOutputDevice& Ar)
+{
+	// required by IInputDevice interface
+	return false;
+}
+
+void IGameInputDeviceInterface::SetChannelValue(int32 ControllerId, FForceFeedbackChannelType ChannelType, float Value)
+{
+#if GAME_INPUT_SUPPORT
+	IPlatformInputDeviceMapper& DeviceMapper = IPlatformInputDeviceMapper::Get();
+	FPlatformUserId UserId = PLATFORMUSERID_NONE;
+	FInputDeviceId DeviceId = INPUTDEVICEID_NONE;
+	DeviceMapper.RemapControllerIdToPlatformUserAndDevice(ControllerId, UserId, DeviceId);
+
+	if (!UserId.IsValid())
+	{
+		return;
+	}
+
+	GameInputRumbleParams RumbleParams = {};
+
+	switch (ChannelType)
+	{
+	case FForceFeedbackChannelType::LEFT_LARGE:
+		RumbleParams.lowFrequency = FMath::Clamp(Value, 0.0f, 1.0f);
+		break;
+	case FForceFeedbackChannelType::LEFT_SMALL:
+		RumbleParams.leftTrigger = FMath::Clamp(Value, 0.0f, 1.0f);
+		break;
+	case FForceFeedbackChannelType::RIGHT_LARGE:
+		RumbleParams.highFrequency = FMath::Clamp(Value, 0.0f, 1.0f);
+		break;
+	case FForceFeedbackChannelType::RIGHT_SMALL:
+		RumbleParams.rightTrigger = FMath::Clamp(Value, 0.0f, 1.0f);
+		break;
+	}
+
+	// Send a rumble event for every input device that is mapped to this platform user
+	for (const TSharedPtr<FGameInputDeviceContainer>& KnownDevice : DeviceData)
+	{
+		if (DeviceMapper.GetUserForInputDevice(KnownDevice->GetDeviceId()) == UserId)
+		{
+			if (IGameInputDevice* Device = KnownDevice->GetGameInputDevice())
+			{
+				Device->SetRumbleState(&RumbleParams);
+			}
+		}
+	}
+
+#endif	// #if GAME_INPUT_SUPPORT
+}
+
+void IGameInputDeviceInterface::SetChannelValues(int32 ControllerId, const FForceFeedbackValues& Values)
+{
+#if GAME_INPUT_SUPPORT
+	// TODO: Allow native input device id's to FInputDeviceId's in the platform input device mapper
+	IPlatformInputDeviceMapper& DeviceMapper = IPlatformInputDeviceMapper::Get();
+	FPlatformUserId UserId = PLATFORMUSERID_NONE;
+	FInputDeviceId DeviceId = INPUTDEVICEID_NONE;
+	DeviceMapper.RemapControllerIdToPlatformUserAndDevice(ControllerId, UserId, DeviceId);
+
+	if (!UserId.IsValid())
+	{
+		return;
+	}
+
+	GameInputRumbleParams RumbleParams = {};
+
+	RumbleParams.lowFrequency = FMath::Clamp(Values.LeftLarge, 0.0f, 1.0f);
+	RumbleParams.leftTrigger = FMath::Clamp(Values.LeftSmall, 0.0f, 1.0f);
+	RumbleParams.highFrequency = FMath::Clamp(Values.RightLarge, 0.0f, 1.0f);
+	RumbleParams.rightTrigger = FMath::Clamp(Values.RightSmall, 0.0f, 1.0f);
+
+	// Send a rumble event for every input device that is mapped to this platform user
+	for (const TSharedPtr<FGameInputDeviceContainer>& KnownDevice : DeviceData)
+	{
+		if (DeviceMapper.GetUserForInputDevice(KnownDevice->GetDeviceId()) == UserId)
+		{
+			if (IGameInputDevice* Device = KnownDevice->GetGameInputDevice())
+			{
+				Device->SetRumbleState(&RumbleParams);
+			}
+		}
+	}
+#endif	// #if GAME_INPUT_SUPPORT
+}
+
+bool IGameInputDeviceInterface::IsGamepadAttached() const
+{
+#if GAME_INPUT_SUPPORT
+	// We will treat both Gamepads and Controller's as "Gamepads" as far as UE is concerned. 
+	return CurrentlyConnectedDeviceTypes & GameInputKindGamepad || CurrentlyConnectedDeviceTypes & GameInputKindController;
+#else
+	return false;
+#endif	// GAME_INPUT_SUPPORT
+}
+
+#if GAME_INPUT_SUPPORT
+
+GameInputKind IGameInputDeviceInterface::GetCurrentGameInputKindSupport() const
+{
+	GameInputKind RegisterInputKindMask = GameInputKindUnknown;
+	
+	const UGameInputPlatformSettings* PlatformSettings = UGameInputPlatformSettings::Get();
+	
+	if (PlatformSettings->bProcessController)
+	{
+		RegisterInputKindMask |= (GameInputKindController | GameInputKindControllerAxis | GameInputKindControllerButton | GameInputKindControllerSwitch);
+	}
+
+	if (PlatformSettings->bProcessRawInput)
+	{
+		RegisterInputKindMask |= GameInputKindRawDeviceReport;
+	}
+
+	if (PlatformSettings->bProcessGamepad)
+	{
+		RegisterInputKindMask |= GameInputKindGamepad;
+	}
+
+	if (PlatformSettings->bProcessKeyboard)
+	{
+		RegisterInputKindMask |= GameInputKindKeyboard;
+	}
+
+	if (PlatformSettings->bProcessMouse)
+	{
+		RegisterInputKindMask |= GameInputKindMouse;
+	}
+
+	// TODO_BH: Future expansion of game input on Xbox
+	/*if (Settings->bProcessArcadeStick)
+	{
+
+	}
+
+	if (Settings->bProcessRacingWheel)
+	{
+
+	}	
+	*/
+
+	return RegisterInputKindMask;
+}
+
+bool IGameInputDeviceInterface::BindDeviceStatusCallbacks()
+{
+	if (!GameInput)
+	{
+		return false;
+	}
+	
+	// First, check if we are already bound to something. If so, unbind it
+	if (ConnectionChangeCallbackToken)
+	{
+		GameInput->UnregisterCallback(ConnectionChangeCallbackToken, UINT64_MAX);
+		ConnectionChangeCallbackToken = GAMEINPUT_INVALID_CALLBACK_TOKEN_VALUE;
+	}
+
+	// Hook up device callbacks 
+	auto DeviceCallbackFn = [](GameInputCallbackToken CallbackToken, void* Context, IGameInputDevice* Device, uint64 Timestamp, GameInputDeviceStatus CurrentStatus, GameInputDeviceStatus PreviousStatus)
+	{
+		static_cast<IGameInputDeviceInterface*>(Context)->OnDeviceConnectionChanged(CallbackToken, Device, Timestamp, CurrentStatus, PreviousStatus);
+	};
+
+	// TODO: should this be a member variable that subclasses can override? That might be desirable on PC
+	static constexpr GameInputDeviceStatus DeviceStatusMask = (GameInputDeviceNoStatus | GameInputDeviceConnected);
+	
+	const GameInputKind RegisterInputKindMask = GetCurrentGameInputKindSupport();
+
+	UE_LOG(LogGameInput, Log, TEXT("Registering Device Callback for GameInputKind: '%s'. Listening for Device Status: '%s'."), *UE::GameInput::LexToString(RegisterInputKindMask), *UE::GameInput::LexToString(DeviceStatusMask));
+	
+	GameInput->RegisterDeviceCallback(nullptr, RegisterInputKindMask, DeviceStatusMask, GameInputBlockingEnumeration, this, DeviceCallbackFn, &ConnectionChangeCallbackToken);
+
+	// get notified when the app is constrained & unconstrained
+	FCoreDelegates::ApplicationWillDeactivateDelegate.AddRaw(this, &IGameInputDeviceInterface::OnAppConstrained);
+	FCoreDelegates::ApplicationHasReactivatedDelegate.AddRaw(this, &IGameInputDeviceInterface::OnAppUnconstrained);
+
+	// We are successful if the callback token is valid
+	return ConnectionChangeCallbackToken != GAMEINPUT_INVALID_CALLBACK_TOKEN_VALUE;
+}
+
+void IGameInputDeviceInterface::OnDeviceConnectionChanged(GameInputCallbackToken CallbackToken, IGameInputDevice* Device, uint64 Timestamp, GameInputDeviceStatus CurrentStatus, GameInputDeviceStatus PreviousStatus)
+{
+	FScopeLock Lock(&DeviceInfoCS);
+
+	UE_LOG(LogGameInput, Log, TEXT("Device Connection Changed from '%s' to '%s'"), *UE::GameInput::LexToString(PreviousStatus), *UE::GameInput::LexToString(CurrentStatus));
+	
+	// This event may come in async from GameInput from outside the game thread, so defer it until Tick so we can guarantee game thread access
+	IGameInputDeviceInterface::FDeferredDeviceConnectionChanges Event = {};
+	Event.Device = Device;
+	Event.Timestamp = Timestamp;
+	Event.Status = UE::GameInput::DeviceStateToConnectionState(CurrentStatus, PreviousStatus);
+	Event.CurrentStatus = CurrentStatus;
+	Event.PreviousStatus = PreviousStatus;
+
+	DeferredDeviceConnectionChanges.Emplace(Event);
+}
+
+void IGameInputDeviceInterface::ProcessDeferredDeviceConnectionChanges()
+{
+	// We only want to actually handle device connection events in the game thread
+	// because a lot of listeners will be in game or ui code that reference the FSlateApplication::Get function.
+	check(IsInGameThread());
+	if (!DeferredDeviceConnectionChanges.IsEmpty())
+	{
+		for (const FDeferredDeviceConnectionChanges& Event : DeferredDeviceConnectionChanges)
+		{
+			// No status means that the device has been disconnected
+			if (Event.Status == EInputDeviceConnectionState::Disconnected)
+			{
+				HandleDeviceDisconnected(Event.Device, Event.Timestamp);
+			}
+			else if (Event.Status == EInputDeviceConnectionState::Connected)
+			{
+				HandleDeviceConnected(Event.Device, Event.Timestamp);
+			}
+			else
+			{
+				UE_LOG(LogGameInput, Error, TEXT("Unexpected state change for device %s (%s -> %s"), *UE::GameInput::LexToString(Event.Device), *UE::GameInput::LexToString(Event.PreviousStatus), *UE::GameInput::LexToString(Event.CurrentStatus));
+			}
+		}
+		DeferredDeviceConnectionChanges.Reset();
+	}
+}
+
+void IGameInputDeviceInterface::OnAppConstrained()
+{
+	bIsAppCurrentlyConstrained = true;
+}
+
+void IGameInputDeviceInterface::OnAppUnconstrained()
+{
+	bIsAppCurrentlyConstrained = false;
+}
+
+void IGameInputDeviceInterface::DetermineStateAfterFirstUnconstrainedUpdate()
+{
+	// This should only be called on the first update after coming back from being constrained
+	ensure(bIsAppCurrentlyConstrained && !bWasAppConstrainedLastTick);
+
+	for (const TSharedPtr<FGameInputDeviceContainer>& KnownDevice : DeviceData)
+	{
+		KnownDevice->ClearInputState(GameInput);
+	}	
+}
+
+FGameInputDeviceContainer* IGameInputDeviceInterface::GetDeviceData(IGameInputDevice* InDevice)
+{
+	// Check if we already know about the given device
+	if (!InDevice)
+	{
+		return nullptr;
+	}
+
+	// Check if we have seen this device's APP_LOCAL_DEVICE_ID before. The IGameInputDevice* could have been invalidated if the device 
+	// is being re-connected, but it will have the APP_LOCAL_DEVICE_ID would be the same.
+	const GameInputDeviceInfo* Info = InDevice->GetDeviceInfo();
+
+	FGameInputDeviceContainer* RetVal = nullptr;
+
+	for (const TSharedPtr<FGameInputDeviceContainer>& KnownDevice : DeviceData)
+	{
+		const APP_LOCAL_DEVICE_ID KnownDeviceId = KnownDevice->GetGameInputDeviceId();
+
+		if (KnownDevice->GetGameInputDevice() == InDevice ||
+			(Info && FMemory::Memcmp(&Info->deviceId, &KnownDeviceId, sizeof(KnownDeviceId)) == 0))
+		{
+			RetVal = KnownDevice.Get();
+			break;
+		}
+	}
+
+	return RetVal;
+}
+
+FGameInputDeviceContainer* IGameInputDeviceInterface::GetOrCreateDeviceData(IGameInputDevice* InDevice)
+{
+	// Check if we already know about the given device
+	if (!InDevice)
+	{
+		return nullptr;
+	}
+
+	// Check if we already have some device data about this...
+	if (FGameInputDeviceContainer* ExistingDevice = GetDeviceData(InDevice))
+	{
+		return ExistingDevice;
+	}
+
+	// ... if not, then create a new one.
+	return CreateDeviceData(InDevice);
+}
+
+void IGameInputDeviceInterface::EnumerateCurrentlyConnectedDeviceTypes()
+{
+	if (!UE::GameInput::CVarGameInputEnumerateDeviceTypeOnConnection.GetValueOnAnyThread())
+	{
+		return;
+	}
+
+	CurrentlyConnectedDeviceTypes = GameInputKindUnknown;
+
+	// Check all of our devices and their supported input flags
+	for (TSharedPtr<FGameInputDeviceContainer>& KnownDevice : DeviceData)
+	{
+		if (KnownDevice)
+		{
+			if (IGameInputDevice* Device = KnownDevice->GetGameInputDevice())
+			{
+				if (const GameInputDeviceInfo* Info = Device->GetDeviceInfo())
+				{
+					CurrentlyConnectedDeviceTypes |= Info->supportedInput;
+				}
+			}
+		}
+	}
+}
+
+#endif	// GAME_INPUT_SUPPORT
diff --git a/Engine/Plugins/Runtime/GameInput/Source/GameInputBase/Public/GameInputBaseIncludes.h b/Engine/Plugins/Runtime/GameInput/Source/GameInputBase/Public/GameInputBaseIncludes.h
new file mode 100644
--- /dev/null
+++ b/Engine/Plugins/Runtime/GameInput/Source/GameInputBase/Public/GameInputBaseIncludes.h
@@ -0,0 +1,22 @@
+// Copyright Epic Games, Inc. All Rights Reserved.
+
+#pragma once
+
+#if GAME_INPUT_SUPPORT
+
+// #include "Microsoft/AllowMicrosoftPlatformTypes.h"
+//
+// THIRD_PARTY_INCLUDES_START
+//
+// #include "Microsoft/COMPointer.h"
+// #include <GameInput.h>
+//
+// #if PLATFORM_WINDOWS
+// 	#include <windef.h>	// For APP_LOCAL_DEVICE_ID
+// #endif	// PLATFORM_WINDOWS
+//
+// THIRD_PARTY_INCLUDES_END
+//
+// #include "Microsoft/HideMicrosoftPlatformTypes.h"
+
+#endif	// GAME_INPUT_SUPPORT
\ No newline at end of file
diff --git a/Engine/Plugins/Runtime/GameInput/Source/GameInputBase/Public/GameInputBaseModule.h b/Engine/Plugins/Runtime/GameInput/Source/GameInputBase/Public/GameInputBaseModule.h
new file mode 100644
--- /dev/null
+++ b/Engine/Plugins/Runtime/GameInput/Source/GameInputBase/Public/GameInputBaseModule.h
@@ -0,0 +1,35 @@
+// Copyright Epic Games, Inc. All Rights Reserved.
+
+#pragma once
+
+#include "Modules/ModuleInterface.h"
+
+#if GAME_INPUT_SUPPORT
+struct IGameInput;
+#endif
+
+class GAMEINPUTBASE_API FGameInputBaseModule : public IModuleInterface
+{
+public:
+
+	static FGameInputBaseModule& Get();
+
+	/** Returns true if this module is loaded (aka available) by the FModuleManager */
+	static bool IsAvailable();
+
+	//~ Begin IModuleInterface interface
+	virtual void StartupModule() override;
+	virtual void ShutdownModule() override;
+	//~ End IModuleInterface interface
+
+#if GAME_INPUT_SUPPORT
+	/** 
+	* Pointer to the static IGameInput that is created upon module startup.
+	*/
+	static IGameInput* GetGameInput();
+#endif
+
+protected:
+
+	void InitializeGameInputKeys();
+};
\ No newline at end of file
diff --git a/Engine/Plugins/Runtime/GameInput/Source/GameInputBase/Public/GameInputDeveloperSettings.h b/Engine/Plugins/Runtime/GameInput/Source/GameInputBase/Public/GameInputDeveloperSettings.h
new file mode 100644
--- /dev/null
+++ b/Engine/Plugins/Runtime/GameInput/Source/GameInputBase/Public/GameInputDeveloperSettings.h
@@ -0,0 +1,437 @@
+// Copyright Epic Games, Inc. All Rights Reserved.
+
+#pragma once
+
+#include "Engine/DeveloperSettings.h"
+#include "Engine/PlatformSettings.h"
+
+#include "GameInputDeveloperSettings.generated.h"
+
+/**
+ * Represents a single unique device that Microsoft's Game Input API detects.
+ * Each device has a Vendor ID and  Product ID. It is possible (and common) for Game Input devices
+ * to have the same VendorId's.
+ *
+ * This is used so that we can easily make a map of each unique device to some metadata about it's controller
+ * bindings.
+ *
+ * If you are unsure about what the vendor or product ID's are, then you can look at the logs
+ * and they will be there in the LogGameInput if you search for "VendorId" or "ProductId".
+ *
+ * You can also place a breakpoint on the device connection status changed to see the value there from the GameInputDeviceInfo.
+ */
+USTRUCT()
+struct GAMEINPUTBASE_API FGameInputDeviceIdentifier
+{
+	GENERATED_BODY()
+
+	FGameInputDeviceIdentifier();
+	FGameInputDeviceIdentifier(uint16 InVendorId, uint16 InProductId);
+
+	/** The Vendor ID of this device */
+	UPROPERTY(EditAnywhere, Config, Category = "Input Device")
+	uint16 VendorId;
+
+	/** The unique Product ID of this device */
+	UPROPERTY(EditAnywhere, Config, Category = "Input Device")
+	uint16 ProductId;
+
+	GAMEINPUTBASE_API friend uint32 GetTypeHash(const FGameInputDeviceIdentifier& InId);
+	bool operator==(const FGameInputDeviceIdentifier& Other) const;
+	bool operator!=(const FGameInputDeviceIdentifier& Other) const;
+
+	/**
+	 * Returns a string representing this Game Input device identifier.
+	 *
+	 * For example:
+	 *		"VendorId: XXX ProductId XXX"
+	 */
+	FString ToString() const;
+};
+
+/**
+* Configuration of an individual axis from a Game Input Controller device type.
+* 
+* These settings would be used for one individual reading from a "GetControllerAxisState" call.
+* 
+* @see FGameInputControllerDeviceProcessor::ProcessControllerAxisState
+*/
+USTRUCT()
+struct GAMEINPUTBASE_API FGameInputControllerAxisData
+{
+	GENERATED_BODY()
+	
+	FGameInputControllerAxisData() = default;
+
+	UPROPERTY(EditAnywhere, Config, Category = "Game Input Device")
+	FName KeyName = NAME_None;
+
+	UPROPERTY(EditAnywhere, Config, Category = "Game Input Device")
+	float DeadZone = (7849.0f / 32768.0f);
+
+	/** Scalar that the input value will be multiplied by before being reported to the input system. */
+	UPROPERTY(EditAnywhere, Config, Category = "Game Input Device")
+	float Scalar = 1.0f;
+
+	/** 
+	* Set this to true if this axis represents a positive and negative value packed into one float.
+	* For example, a trigger will be a single value between 0.0 and 1.0, where 0.0 is unactuated and 1.0 is full actuated.
+	* However, something like a thumb stick may pack both positive and negative values into its range, from -1.0 to +1.0 
+	* where -1.0 is all the way to the left, and +1.0 is all the way to the right.
+	*/
+	UPROPERTY(EditAnywhere, Config, Category = "Game Input Device")
+	bool bIsPackedPositveAndNegative = false;
+};
+
+/***
+ * Options for how to translate the raw uint8 RawInput report data to the engine.
+ */
+UENUM()
+enum class ERawDeviceReportTranslationBehavior : uint8
+{
+	// Treat the value as a trigger, mapping it from 0.0 to +1.0
+	TreatAsTrigger = 0,
+	
+	// Treat the value as an analog value, mapping it from -1.0 to +1.0
+	TreatAsAnalog = 1,
+	
+	// Treats the raw value as a button. If the value is 0, consider it not pressed. If the value is non-zero, consider it pressed.
+	TreatAsButton = 2,
+
+	// Treat the raw value as a bitmask of different buttons. Treats each bit in the raw uint8 value as its own button.
+	// 
+	// If the bit is 0, then consider it not pressed. If the bit is non-zero (1) then consider it pressed. 
+	TreatAsButtonBitmask = 3,
+
+	// Treat this as two raw values, where one axis is the "lower" bits of an int16 and 
+	// one is the "higher" bits of the int16. 
+	TreatAsPackedAxisPair = 4,
+};
+
+/**
+* A configuration of a single Raw Device Report index. 
+* This information will be used to map the raw device report information 
+* of a given index to something that the Unreal Engine can process with its
+* message handler.
+* 
+* @see FGameInputRawDeviceProcessor::ProcessInput
+*/
+USTRUCT()
+struct GAMEINPUTBASE_API FGameInputRawDeviceReportData
+{
+	GENERATED_BODY()
+	
+	FGameInputRawDeviceReportData() = default;
+
+	/**
+	 * The name of the FKey that this raw device report is associated with when treated as an analog or trigger value
+	 * 
+	 * Note: This value only matters for Translation Behaviors other then the button bitmask.
+	 */
+	UPROPERTY(EditAnywhere, Config, Category = "Game Input Device", meta=(EditCondition="TranslationBehavior != ERawDeviceReportTranslationBehavior::TreatAsButtonBitmask", EditConditionHides))
+	FName KeyName = NAME_None;
+
+	/**
+	 * Options for how we should interpret the raw uint8 value from RawInput when telling the engine about it
+	 */
+	UPROPERTY(EditAnywhere, Config, Category = "Game Input Device")
+	ERawDeviceReportTranslationBehavior TranslationBehavior = ERawDeviceReportTranslationBehavior::TreatAsTrigger;
+
+	/**
+	 * The deadzone to use when translating the value to an analog input between -1.0 and +1.0.
+	 * 
+	 * Normally a value between 0-10 is about where you want it for most devices.
+	 */
+	UPROPERTY(EditAnywhere, Config, Category = "Game Input Device", meta=(EditCondition="TranslationBehavior == ERawDeviceReportTranslationBehavior::TreatAsAnalog", EditConditionHides))
+	uint8 AnalogDeadzone = 5;
+
+	/**
+	* A scalar that gets applied to the raw float value for Analog or trigger values. Most of the time this would be used to simply negate values
+	*/
+	UPROPERTY(EditAnywhere, Config, Category = "Game Input Device", meta=(EditCondition="TranslationBehavior == ERawDeviceReportTranslationBehavior::TreatAsAnalog || TranslationBehavior == ERawDeviceReportTranslationBehavior::TreatAsTrigger", EditConditionHides))
+	float Scalar = 1.0f;
+
+	/**
+	 * When treating this raw input as a button we will map the key value of this map to a bit on the raw input value.
+	 * I.e. a key value of "1" will look at the first bit to see if it is set. If it is, then send a controller button event
+	 * for the associated FName
+	 */
+	UPROPERTY(EditAnywhere, Config, Category = "Game Input Device", meta=(EditCondition="TranslationBehavior == ERawDeviceReportTranslationBehavior::TreatAsButtonBitmask", EditConditionHides))
+	TMap<int32, FName> ButtonBitMaskMappings;
+
+
+	UPROPERTY(EditAnywhere, Config, Category = "Game Input Device", meta = (EditCondition = "TranslationBehavior == ERawDeviceReportTranslationBehavior::TreatAsPackedAxisPair", EditConditionHides))
+	uint8 LowerBitAxisIndex = 0;
+	
+	UPROPERTY(EditAnywhere, Config, Category = "Game Input Device", meta = (EditCondition = "TranslationBehavior == ERawDeviceReportTranslationBehavior::TreatAsPackedAxisPair", EditConditionHides))
+	uint8 HigherBitAxisIndex = 0;
+};
+
+/**
+* Configurable data about a unique type of Game Input device. These settings
+* will be used at runtime to determine the behavior of any device whose Vendor and
+* Product ID match this one.
+*/
+USTRUCT()
+struct GAMEINPUTBASE_API FGameInputDeviceConfiguration
+{
+	GENERATED_BODY()
+
+	FGameInputDeviceConfiguration();
+
+	/** The unique device vendor and product ID that can be used to ID this item at runtime. */
+	UPROPERTY(EditAnywhere, Config, Category = "Device Settings")
+	FGameInputDeviceIdentifier DeviceIdentifier;
+
+#if WITH_EDITORONLY_DATA
+	/**
+	* An editor-only description of this input device. 
+	* Only used for if you want to describe your device and make the project settings easier to read 
+	*/
+	UPROPERTY(EditAnywhere, Config, Category = "Device Settings", meta=(DisplayName="Description (Editor Only)"))
+	FString Description = TEXT("");
+#endif
+	
+	/** 
+	* If true, then the hardware device Id specific on this configuration will be used instead of
+	* whatever the Game Input SDK tells us this device is. Use this if want to get additional
+	* messaging from the Input Device Subsystem when this specific device is used.
+	* 
+	* For example, some third party controller might be an "Xbox One" controller type, but 
+	* it is really some specially manufactured hardware for your game that you would like 
+	* to have access to at the Gameplay layer in UE.
+	*/
+	UPROPERTY(EditAnywhere, Config, Category = "Device Settings", meta = (InlineEditConditionToggle))
+	bool bOverrideHardwareDeviceIdString;
+
+	/**
+	* The name of this device that should be used for FInputDeviceScope's to determine 
+	* when it has been connected or input events came from it.
+	*/
+	UPROPERTY(EditAnywhere, Config, Category = "Device Settings", meta=(EditCondition="bOverrideHardwareDeviceIdString"))
+	FString OverriddenHardwareDeviceId;
+
+	/** If true, this device will attempt to process Game Input Button states. */
+	UPROPERTY(EditAnywhere, Config, Category = "Device Settings|Controller", meta = (InlineEditConditionToggle))
+	bool bProcessControllerButtons;
+
+	/** 
+	* If true, this device will attempt to process Game Input Switch (aka DPad) state.
+	* 
+	* If the device has it's DPad being processed as buttons already, you may want to turn this off.
+	*/
+	UPROPERTY(EditAnywhere, Config, Category = "Device Settings|Controller")
+	bool bProcessControllerSwitchState;
+	
+	/**
+	 * If true, this device will attempt to process Game Input Controller Axis mappings.
+	 */
+	UPROPERTY(EditAnywhere, Config, Category = "Device Settings|Controller", meta = (InlineEditConditionToggle))
+	bool bProcessControllerAxis;
+
+	/**
+	* A map of uint32 button index to an associated FName Unreal gamepad key name.
+	* 
+	* These key values should be ([button index] << 1) i.e, powers of two.
+	*/
+	UPROPERTY(EditAnywhere, Config, Category = "Device Settings|Controller", meta=(EditCondition="bProcessControllerButtons"))
+	TMap<uint32, FName> ControllerButtonMappingData;
+
+	/**
+	* A map of uint32 button index to some data about how we should treat the raw axis value
+	*/
+	UPROPERTY(EditAnywhere, Config, Category = "Device Settings|Controller", meta=(EditCondition="bProcessControllerAxis"))
+	TMap<uint32, FGameInputControllerAxisData> ControllerAxisMappingData;
+
+	/**
+	 * If true, this device will attempt to process Game Input Raw Report mappings.
+	 */
+	UPROPERTY(EditAnywhere, Config, Category = "Device Settings|Raw Report")
+	bool bProcessRawReportData;
+
+	/** 
+	* The raw report reading ID ( GameInputRawDeviceReportInfo::id ).
+	* 
+	* We need to know this value so that we can confirm that the reading we are receiving is for this device configuration.
+	*/
+	UPROPERTY(EditAnywhere, Config, Category = "Device Settings|Raw Report", meta=(EditCondition="bProcessRawReportData"))
+	uint32 RawReportReadingId;
+	
+	/**
+	* A map of uint32 raw report index to some data about how we should treat that raw data.
+	*/
+	UPROPERTY(EditAnywhere, Config, Category = "Device Settings|Raw Report", meta=(EditCondition="bProcessRawReportData"))
+	TMap<int32, FGameInputRawDeviceReportData> RawReportMappingData;
+};
+
+struct GameInputDeviceInfo;
+
+/**
+ * Settings for Game Input that you may want to set per-platform.
+ */
+UCLASS(config = Input, DefaultConfig)
+class GAMEINPUTBASE_API UGameInputPlatformSettings : public UPlatformSettings
+{
+	GENERATED_BODY()
+	
+	friend class UGameInputDeveloperSettings;
+
+protected:
+
+	virtual void InitializePlatformDefaults() override;
+	
+public:
+
+	/** Returns a pointer to Game Input settings on the current platform. */
+	static UGameInputPlatformSettings* Get();
+
+	/**
+	 * If true Game Input will process controller axis, button, and switch events
+	 *
+	 * Default: False
+	 */
+	UPROPERTY(Config, EditAnywhere, Category="Processing Options")
+	bool bProcessController = false;
+
+	/** 
+	* If true, then we will process the GameInputKindRawDeviceReport type.
+	*
+	* Default: False
+	*/
+	UPROPERTY(Config, EditAnywhere, Category="Processing Options")
+	bool bProcessRawInput = false;
+	
+	/**
+	* If true, then we will only process controller devices whose vendor/product ID is in the 
+	* DeviceConfigurations array. This applies to both Controller and Raw input kinds. 
+	* 
+	* Default: True
+	*/
+	UPROPERTY(Config, EditAnywhere, Category="Processing Options", meta=(EditCondition="bProcessController || bProcessRawInput"))
+	bool bSpecialDevicesRequireExplicitDeviceConfiguration = true;
+
+	/** 
+	* If true, then we will process the GameInputKindGamepad type.
+	* 
+	* Default: True
+	*/
+	UPROPERTY(Config, EditAnywhere, Category="Processing Options")
+	bool bProcessGamepad = true;
+
+	/** 
+	* If true, then we will process the GameInputKindKeyboard type.
+	* 
+	* Note: You likely do not want this on for Windows targets, as those input events are already processed via the 
+	* WindowsApplication. This should really only be used for console targets.
+	* 
+	* Default: True
+	*/
+	UPROPERTY(Config, EditAnywhere, Category="Processing Options")
+	bool bProcessKeyboard = true;
+
+	/** 
+	* If true, then we will process the GameInputKindMouse type.
+	* 
+	* Note: You likely do not want this on for Windows targets, as those input events are already processed via the 
+	* WindowsApplication. This should really only be used for console targets.
+	* 
+	* Default: True
+	*/
+	UPROPERTY(Config, EditAnywhere, Category="Processing Options")
+	bool bProcessMouse = true;
+};
+
+/**
+* Settings related to theGame Input device interface. 
+* 
+* These will allow you to enable and disable specific types of input devices within Game Input
+* as well as configure key mappings for generic controller types based on their unique vendor/product ID's.
+*/
+UCLASS(config = Input, DefaultConfig, meta = (DisplayName = "Game Input Plugin Settings"))
+class GAMEINPUTBASE_API UGameInputDeveloperSettings : public UObject
+{
+	GENERATED_BODY()
+public:
+
+#if GAME_INPUT_SUPPORT
+	/**
+	* Finds unique device configuration data for the given Game Input Device. This is based
+	* on the vendor/product ID of the game input device
+	* 
+	* @return Device Configuration for the given Game Input device. Null if it is not configured.
+	*/
+	[[nodiscard]] const FGameInputDeviceConfiguration* FindDeviceConfiguration(const GameInputDeviceInfo* const Info) const;
+
+	/**
+	 * Returns a pointer to the device configuration that may exist for the given hardware ID
+	 */
+	const FGameInputDeviceConfiguration* FindDeviceConfiguration(const FGameInputDeviceIdentifier& HardwareID) const;
+#endif	 // GAME_INPUT_SUPPORT
+
+	UGameInputDeveloperSettings(const FObjectInitializer& Initializer);
+
+protected:
+	
+	/**
+	 * Array of devices that you want to specify the behavior for. These could be some
+	 * special input devices for your game like racing wheels, instruments, or other "special"
+	 * peripherals that require some kind of special readings from Game Input.
+	 *
+	 * This is typically for GameInputKindController or GameInputKindRawDeviceReport.
+	 *
+	 * If bSpecialDevicesRequireExplicitDeviceConfiguration is true, only Controller/Raw device kinds
+	 * that are within this array will be allowed to be processed.
+	 */
+	UPROPERTY(EditAnywhere, Config, Category = "Game Input")
+	TArray<FGameInputDeviceConfiguration> DeviceConfigurations;
+
+	/**
+	 * Settings for Game input that you may want to set per-platform.
+	 *
+	 * This is useful to have per-platform because some platforms might already handle the processing
+	 * of some input devices, so you might want to disable them here.
+	 */
+	UPROPERTY(EditAnywhere, Category = "Game Input|Platform Options")
+	FPerPlatformSettings PlatformSpecificSettings;
+
+public:
+	/**
+	* If true, when you have multiple input devices mapped to a single FPlatformUserId,
+	* we will only process events from one of them each frame. This makes it impossible for
+	* analog input values to "over accumulate" and report values higher than +- 1.0
+	*
+	* One example of this issue would be if a user has mapped two game pads to the same platform user
+	* and them presses the right analog stick all the way to the right on both of them at the same time. When this setting
+	* is on, we will only read the value from one of the sticks, resulting in an input value of +1.0. 
+	* With this setting off (which is not recommended) then you would receive an input value of +2.0. 
+	* This could result in having faster rotations around in your game or other non-desirable behaviors. 
+	*
+	* WARNING: Turning this off may result in Slate receiving input events greater than 1.0, which can cause 
+	* undefined behavior in your if it makes the assumption of input being between -1.0 and +1.0.
+	* 
+	* Default: True
+	*/
+	UPROPERTY(Config, EditAnywhere, Category="Game Input")
+	bool bDoNotProcessDuplicateCapabilitiesForSingleUser;
+
+#if WITH_EDITOR
+public:
+	DECLARE_DELEGATE_OneParam(FGameInputSettingChanged, FName);
+	/** A delegate that is called on PostEditChangeProperty with the name of the property that has changed. */
+	FGameInputSettingChanged OnInputSettingChanged;
+	
+	/**
+	 * Returns all known uint32 options for controller button configurations.
+	 * this is useful for making it easier to select in the editor instead of having to manually type everything in
+	 * on a UPROPERTY
+	 * i.e.
+	 *
+	 * meta=(GetOptions="GameInput.GameInputDeveloperSettings.GetControllerButtonMappingDataKeyOptions")
+	 */
+	UFUNCTION()
+	static const TArray<uint32>& GetControllerButtonMappingDataKeyOptions();
+
+protected:
+	virtual void PostEditChangeProperty(FPropertyChangedEvent& PropertyChangedEvent) override;
+#endif	// WITH_EDITOR
+};
\ No newline at end of file
diff --git a/Engine/Plugins/Runtime/GameInput/Source/GameInputBase/Public/GameInputDeviceContainer.h b/Engine/Plugins/Runtime/GameInput/Source/GameInputBase/Public/GameInputDeviceContainer.h
new file mode 100644
--- /dev/null
+++ b/Engine/Plugins/Runtime/GameInput/Source/GameInputBase/Public/GameInputDeviceContainer.h
@@ -0,0 +1,113 @@
+// Copyright Epic Games, Inc. All Rights Reserved.
+
+#pragma once
+
+#include "GameInputDeviceProcessor.h"
+#include "GenericPlatform/GenericApplicationMessageHandler.h"
+
+#if GAME_INPUT_SUPPORT
+
+/**
+* Game Input Device Containers hold any state about a single unique GameInput device.
+* Upon creation, the will create any "Device Processors" that are required based 
+* on the IGameInputDevice's supported GameInputKind. They are ticked each frame by the 
+* owning IGameInputDeviceInterface who created them. 
+* 
+* This class is what will actually poll the GameInput SDK for new readings and keep track
+* of it's state across frames. 
+*/
+class GAMEINPUTBASE_API FGameInputDeviceContainer
+{
+public:
+	FGameInputDeviceContainer(
+		const TSharedRef<FGenericApplicationMessageHandler>& InMessageHandler,
+		IGameInputDevice* InDevice,
+		GameInputKind InAllowedGameInputKinds,
+		FPlatformUserId InUserId = PLATFORMUSERID_NONE, 
+		FInputDeviceId InDeviceId = INPUTDEVICEID_NONE);
+
+	virtual ~FGameInputDeviceContainer() = default;
+
+	/** Create any device processors necessary for this Game Input device to process input and send input events */
+	void InitalizeDeviceProcessors();
+
+	/** Set the message handler that this Device Container should send events through. */
+	void SetMessageHandler(const TSharedRef<FGenericApplicationMessageHandler>& InMessageHandler);
+	
+	/**
+	 * Process any input events from the given reading and send events to a message handler
+	 * 
+	 * @param 	GameInput								Pointer to the game input instance
+	 * 
+	 * @param 	CurrentSupportedKind					Currently supported game input reading kinds that we are allowed to process. This will be used to get the 
+	 *													current readings from GameInput with GameInput->GetNextReading
+	 * 
+	 * @param 	ProcessedKindsForPlatformUserThisFrame	A bitmask of any game input kinds that have already been processed by the platform user this frame.
+	 *													If the UGameInputDeveloperSettings::bDoNotProcessDuplicateCapabilitiesForSingleUser flag is true,
+	 *													then will skip processors whose input kind is set in this bitmask.
+	 * 
+	 * Returns a bitmask of what GameInputKind's were processed this frame.
+	 */
+	const GameInputKind ProcessInput(IGameInput* GameInput, const GameInputKind CurrentSupportedKind, const GameInputKind ProcessedKindsForPlatformUserThisFrame);
+
+	/** Reset any input state that is necessary. This would be called when the application is no longer constrained for example */
+	void ClearInputState(IGameInput* GameInput);
+
+	IGameInputDevice* GetGameInputDevice() const;
+	void SetGameInputDevice(IGameInputDevice* InDevice);
+
+	/** Returns the unique local device ID of this IGameInputDevice from the GameInput SDK. */
+	APP_LOCAL_DEVICE_ID GetGameInputDeviceId() const;
+
+	void SetPlatformUserId(const FPlatformUserId InUserId);
+	FPlatformUserId GetPlatformUserId() const;
+
+	void SetInputDeviceId(const FInputDeviceId InDeviceId);
+	FInputDeviceId GetDeviceId() const;
+
+	uint64 GetLastReadingTimestamp() const { return LastReadingTimestamp; }
+
+	/** 
+	* Returns the number of processors that this device container currently has.
+	*
+	* Note: If this is zero, then this device can't possibly fire any input events.
+	*/
+	const int32 GetNumberOfProcessors() const { return Processors.Num(); }
+
+protected:
+
+	/** Create any device processors necessary for this Game Input device to process input and send input events */
+	virtual void InitalizeDeviceProcessors_Impl();
+	
+	/** Message handler that we can use to tell the engine about input events */
+	TSharedRef<FGenericApplicationMessageHandler> MessageHandler;
+
+	/** The owning Game Input device that this device is associated with */
+	IGameInputDevice* Device;
+
+	/** The kinds of Input that are allowed to be created for this device. */
+	GameInputKind AllowedGameInputKinds;
+
+	/** The owning Platform User ID that this device is associated with */
+	FPlatformUserId UserId;
+
+	/** The Input Device ID that this Game Input Device is associated with */
+	FInputDeviceId AssignedDeviceId;
+
+	/** The unique ID that is associated with this device from the GameInput API */
+	APP_LOCAL_DEVICE_ID LocalDeviceId;
+
+	/** Input Device processors to handle input for any specific types required (gamepad, racing wheel, KBM, etc) */
+	TArray<TSharedPtr<IGameInputDeviceProcessor>> Processors;
+
+	/** The last reading that was processed by this device */
+	TComPtr<IGameInputReading> LastReading;
+
+	/** Timestamp that can be used to check if we should ignore a reading from the past or not */
+	uint64 IgnoreReadingTimestamp;
+
+	/** Timestamp of the last reading processed by this device */
+	uint64 LastReadingTimestamp;
+};
+
+#endif	// GAME_INPUT_SUPPORT
\ No newline at end of file
diff --git a/Engine/Plugins/Runtime/GameInput/Source/GameInputBase/Public/GameInputDeviceProcessor.h b/Engine/Plugins/Runtime/GameInput/Source/GameInputBase/Public/GameInputDeviceProcessor.h
new file mode 100644
--- /dev/null
+++ b/Engine/Plugins/Runtime/GameInput/Source/GameInputBase/Public/GameInputDeviceProcessor.h
@@ -0,0 +1,449 @@
+// Copyright Epic Games, Inc. All Rights Reserved.
+
+#pragma once
+
+#include "Math/MathFwd.h"
+#include "CoreTypes.h"
+#include "Containers/Map.h"
+#include "Misc/CoreMiscDefines.h"	// For FInputDeviceId/FPlatformUserId
+#include "GenericPlatform/GenericApplicationMessageHandler.h"
+
+#if GAME_INPUT_SUPPORT
+THIRD_PARTY_INCLUDES_START
+#include "Microsoft/AllowMicrosoftPlatformTypes.h"
+#include <GameInput.h>
+#include "Microsoft/HideMicrosoftPlatformTypes.h"
+THIRD_PARTY_INCLUDES_END
+#include "Microsoft/COMPointer.h"
+#endif	// GAME_INPUT_SUPPORT
+
+#if GAME_INPUT_SUPPORT
+
+/**
+* A Game Input Device Processor is used to determine the state of a device
+* each frame and send messages to the given Message Handler when we want to. 
+* Each processor represents a single "kind" of Game Input device. 
+* 
+* They are created by the owning GameInput Device Container (GameInputDeviceContainer.h)
+* and polled each frame when possible. Each Human Interface Device can be made up of multiple 
+* Game Input Kinds, so a single device container may have multiple different processors associated with it.
+*/
+class GAMEINPUTBASE_API IGameInputDeviceProcessor
+{
+public:
+	IGameInputDeviceProcessor() = default;
+
+	virtual ~IGameInputDeviceProcessor() = default;
+
+	struct GAMEINPUTBASE_API FGameInputEventParams
+	{
+		/** The current reading for this input event */
+		IGameInputReading* Reading = nullptr;
+
+		/** The previous reading from the last frame. */
+		IGameInputReading* PreviousReading = nullptr;
+
+		/** Device associated with this reading */
+		IGameInputDevice* Device = nullptr;
+
+		/** The message handler to use to send any input events */
+		TSharedPtr<FGenericApplicationMessageHandler> MessageHandler;
+
+		/** The platform user that is associated with the input device that triggered this event */
+		FPlatformUserId PlatformUserId = PLATFORMUSERID_NONE;
+
+		/** The Input device ID of the device that triggered this event. */
+		FInputDeviceId InputDeviceId = INPUTDEVICEID_NONE;
+
+		/** Gets the Input Device Info for this GameInput Device if it is valid. Can be null. */
+		const GameInputDeviceInfo* GetDeviceInfo() const;
+	};
+
+	/** 
+	* Process any input events from the given reading and send events to the message handler
+	* using the given Platform User and Input Device Id's.
+	* 
+	* This can be called multiple times per frame, like if we are running at a lower FPS.
+	* This is because GameInput may accumulate multiple readings to use by the time the game thread gets here.
+	* 
+	* If you have inputs that you only want to process once per frame, then use PostProcessInput instead.
+	* 
+	* @return True if any input events have been processed, false if nothing happened
+	*/
+	virtual bool ProcessInput(const FGameInputEventParams& Params) = 0;
+
+	/**
+	* Called after all current Game Input readings have been processed
+	* this frame and there are no more readings for a device (GetNextReading returns GAMEINPUT_E_READING_NOT_FOUND). 
+	* This is called regardless of if an accompanying ProcessInput call has happened.
+	* 
+	* This function is useful for dealing with multiple readings in a single frame for analog devices. 
+	* It is recommended that processors override this function and only process analog input events here at
+	* the end of the frame, once, instead of multiple times in ProcessInput. That is because at lower frame rates
+	* or hitches you will end up with more then one GameInput reading in the stack that we need to process. If you were to
+	* process all of those analog inputs in a single frame, the values would accumulate to above or below +-1.0,
+	* which could have unexpected behavior in Slate.
+	* 
+	* @param Params		Game Input event params sent from the owning FGameInputDeviceContainer.
+	*					Params.Reading should always be null here, as there is no current reading.
+	*					Params.PreviousReading will contain the most recent Game Input reading from the input stack.
+	* 
+	* @return			True if any input events have been processed, false if nothing happened
+	*/
+	virtual bool PostProcessInput(const FGameInputEventParams& Params);
+
+	/** 
+	* Clear any input state related to this processor, typically by sending events to the 
+	* message handler that each FKey related to this processor now has a value of 0.
+	*/
+	virtual void ClearState(const FGameInputEventParams& Params) = 0;
+
+	/*
+	* Returns the kind of reading that this processor supports. If the current GameInput reading 
+	* is not included in this bitmask, then this processor will not not have it's "ProcessInput" 
+	* function called.
+	* 
+	* @see FGameInputDeviceContainer::ProcessInput
+	*/
+	virtual GameInputKind GetSupportedReadingKind() const = 0;
+
+protected:	
+
+	/** 
+	* Returns the string hardware device identifier for the given input event params that can be used
+	* to create an FInputDeviceScope
+	*/
+	const FString& GetHardwareDeviceIdentifierName(const IGameInputDeviceProcessor::FGameInputEventParams& Params) const;
+
+	/** A general use function to call the message handler and tell it about a controller analog key being used */
+	void OnControllerAnalog(const FGameInputEventParams& Params, const FName& GamePadKey, float NewAxisValueNormalized, float OldAxisValueNormalized, float DeadZone);
+
+	/**
+	 * Helper function for processing the button states of Game input.
+	 *
+	 * @param CurrentButtonHeldMask			The current state of the button mask to evaluate
+	 * @param PreviousButtonMask		The previous state of the button mask. This value will be modified to store the current button mask after evaluation
+	 * @param RepeatTime				A map of the button index to a time that it was last active, used to determine if the key meets the qualifications for a repeat event
+	 * @param UnrealButtonNameMap		A map of Unreal Engine Gamepad key names to their uint32 mapping from Game Input.
+	 * @param SupportedButtonCount		The maximum number of buttons we can process
+	 */
+	void EvaluateButtonStates(
+		const FGameInputEventParams& Params,
+		const uint32 CurrentButtonHeldMask,
+		uint32& PreviousButtonMask,
+		double* RepeatTime,
+		const TMap<uint32, FGamepadKeyNames::Type>& UnrealButtonNameMap,
+		const uint32 SupportedButtonCount);
+
+	/**
+	* Helper function to processing the state of a Switch position (aka a DPad, like left/right/up/down)
+	* 
+	* @param Params		The GameInput event params
+	* @param CurrentPosition	The current position of the switch
+	* @param PreviousPosition	The previous position of the switch. This will be set to the value of the current switch at the end of this function
+	* @param RepeatTimes		Array of doubles that represent the time at which a switch was last pressed that we can use to check for when to send repeat events
+	*/
+	void EvaluateSwitchState(
+		const FGameInputEventParams& Params,
+		GameInputSwitchPosition CurrentPosition,
+		GameInputSwitchPosition& PreviousPosition,
+		TArray<double>& RepeatTimes);
+};
+
+/**
+* Processor for Gamepad Inputs
+*/
+class GAMEINPUTBASE_API FGameInputGamepadDeviceProcessor : public IGameInputDeviceProcessor
+{
+protected:
+	virtual bool ProcessInput(const FGameInputEventParams& Params) override;
+	virtual bool PostProcessInput(const FGameInputEventParams& Params) override;
+	virtual void ClearState(const FGameInputEventParams& Params) override;
+	virtual GameInputKind GetSupportedReadingKind() const override;
+
+	bool ProcessGamepadAnalogState(const FGameInputEventParams& Params, GameInputGamepadState& GamepadState);
+	virtual bool ProcessGamepadButtonState(const FGameInputEventParams& Params, GameInputGamepadState& GamepadState);
+
+	GameInputGamepadState PreviousState = {};
+
+	uint32 LastButtonHeldMask = 0;
+
+	enum { MaxSupportedButtons = 32 };
+
+	double RepeatTime[MaxSupportedButtons];
+
+	/** 
+	* Keeps track of how many times this gamepad has been processed this frame. 
+	* Every successful processing of button input in ProcessInput will increment this value.
+	* This will get reset at the end of the input frame in PostProcessInput.
+	*/
+	int32 NumReadingsProcessedThisFrame = 0;
+};
+
+struct FGameInputDeviceConfiguration;
+
+/**
+* Processor for Controller Input Types
+* 
+* Controllers are very similar to Gamepad types, and often can just be 
+* a third party controller type that is not natively supported by the "GamepadInputKind".
+* This can include things like instruments, other platform controller types, or other 
+* third party input devices.
+*/
+class GAMEINPUTBASE_API FGameInputControllerDeviceProcessor : public IGameInputDeviceProcessor
+{
+public:
+	FGameInputControllerDeviceProcessor();
+
+protected:
+	virtual bool ProcessInput(const FGameInputEventParams& Params) override;
+	virtual bool PostProcessInput(const FGameInputEventParams& Params) override;
+	virtual void ClearState(const FGameInputEventParams& Params) override;
+	virtual GameInputKind GetSupportedReadingKind() const override;
+
+	/** Processes Controller buttons (often face buttons, like ABXY or Circle/Triangle/Square/X) */
+	bool ProcessControllerButtonState(const FGameInputEventParams& Params, const FGameInputDeviceConfiguration* Config, IGameInputReading* InputReading);
+
+	/**
+	* Process any axis readings on this controller.
+	* This will only process axises that are configured in the given Config.
+	*/
+	bool ProcessControllerAxisState(const FGameInputEventParams& Params, const FGameInputDeviceConfiguration* Config, IGameInputReading* InputReading);
+
+	/**
+	* Process the controller switch (also known as the DPad) state
+	*/
+	bool ProcessControllerSwitchState(const FGameInputEventParams& Params, const FGameInputDeviceConfiguration* Config, IGameInputReading* InputReading);
+
+	/** A map of GameInputLabel types to Unreal Engine FKey names. */
+	static const TMap<GameInputLabel, FGamepadKeyNames::Type>& GetGameInputButtonLabelToUnrealName();
+
+	/** Bitmask of any held buttons from the last frame */
+	uint32 LastButtonHeldMask = 0;
+
+	/** Max number of buttons this processor supports */
+	enum { MaxSupportedButtons = 32 };
+
+	/** Timings of when a key press should be considered a "repeat" key */
+	double RepeatTime[MaxSupportedButtons];
+
+	/** Repeat times for when switches are pressed (aka the DPad) */
+	TArray<double> SwitchRepeatTimes;
+
+	/** Previous frame's axis values */
+	TArray<float> PreviousControllerAxisValues;
+
+	/** Array of switch positions */
+	TArray<GameInputSwitchPosition> PreviousSwitchPositions;
+
+	/**
+	* Keeps track of how many times this gamepad has been processed this frame.
+	* Every successful processing of button input in ProcessInput will increment this value.
+	* This will get reset at the end of the input frame in PostProcessInput.
+	*/
+	int32 NumReadingsProcessedThisFrame = 0;
+};
+
+/**
+* Processor for Keyboard inputs
+*/
+class GAMEINPUTBASE_API FGameInputKeyboardDeviceProcessor : public IGameInputDeviceProcessor
+{
+protected:
+	virtual bool ProcessInput(const FGameInputEventParams& Params) override;
+	virtual void ClearState(const FGameInputEventParams& Params) override;
+	virtual GameInputKind GetSupportedReadingKind() const override;
+
+	void UpdateUnifiedKeyboardState(const FGameInputEventParams& Params, TSet<uint8>& CurrentPressedKeys);
+	
+	inline bool IsKeyPressed(uint8 VirtualKeyCode) const { return LastPressedKeys.Contains(VirtualKeyCode); }
+
+	virtual void SetSimulatedCapsLock(bool bVal);
+
+	TSet<uint8> LastPressedKeys;	
+
+	TMap<uint8, double> KeyRepeatTime;	
+
+	// cannot read actual caps lock state, so necessary to simulate it as a bool toggle
+	bool bSimulatedCapsLock;
+};
+
+
+/**
+* Processor for mouse inputs
+*/
+class GAMEINPUTBASE_API FGameInputMouseDeviceProcessor : public IGameInputDeviceProcessor
+{
+public:
+	explicit FGameInputMouseDeviceProcessor(const TSharedPtr<class ICursor>& InCursor);
+
+protected:	
+	virtual bool ProcessInput(const FGameInputEventParams& Params) override;
+	virtual void ClearState(const FGameInputEventParams& Params) override;
+	virtual GameInputKind GetSupportedReadingKind() const override;
+
+	float CalculateMouseAcceleration(int32 Delta);
+
+	virtual bool CanProcessVirtualMouse() const;
+
+	TSharedPtr<class ICursor> Cursor;
+
+	GameInputMouseState PreviousMouseState;
+
+	FVector2D LastMouseOffset;
+
+	enum { MaxSupportedButtons = 32 };
+	double RepeatTime[MaxSupportedButtons];
+
+	uint32 LastButtonHeldMask;
+};
+
+
+/**
+* Processor for touch inputs
+*/
+class GAMEINPUTBASE_API FGameInputTouchDeviceProcessor : public IGameInputDeviceProcessor
+{
+protected:
+	virtual bool ProcessInput(const FGameInputEventParams& Params) override;
+	virtual void ClearState(const FGameInputEventParams& Params) override;
+	virtual GameInputKind GetSupportedReadingKind() const override;
+
+	struct FTouchData
+	{
+		FTouchData()
+		{
+			TouchId = INDEX_NONE;
+			Position = FVector2D::ZeroVector;
+			Pressure = 0.0f;
+			bHasMoved = false;
+			bIsActive = false;
+		}
+
+		int64 TouchId;
+		FVector2D Position;
+		float Pressure;
+		bool bHasMoved;
+		bool bIsActive;
+		};
+	TArray<FTouchData> PreviousTouchData;
+	int32 ActiveTouchPoints;
+
+	// The max touch X and Y of the touch region to let us scale input based on the size of the touch area
+	int32 MaxTouchX = 1920;
+	int32 MaxTouchY = 1080;
+};
+
+struct FGameInputDeviceConfiguration;
+struct FGameInputRawDeviceReportData;
+
+/**
+* Processor for raw input types (GameInputKindRawDeviceReport)
+* 
+* Raw input is very customizable and is very different per-input device, so it requires a lot of configuration.
+* Most of the time it is special analog inputs on custom third party devices. A good example is the Whammy bar on
+* a guitar controller, or maybe some special shifter on a racing wheel that the Game Input API does not pick up by default.
+*/
+class GAMEINPUTBASE_API FGameInputRawDeviceProcessor : public IGameInputDeviceProcessor
+{
+protected:
+	virtual bool ProcessInput(const FGameInputEventParams& Params) override;
+	virtual bool PostProcessInput(const FGameInputEventParams& Params) override;
+	virtual void ClearState(const FGameInputEventParams& Params) override;
+	virtual GameInputKind GetSupportedReadingKind() const override;
+
+	/**
+	* Reads the current state of the raw input buffer and populates CurrentRawData and PreviousRawData
+	*/
+	const GameInputRawDeviceReportInfo* ReadCurrentRawInputState(const FGameInputEventParams& Params, IGameInputReading* ReadingToUse);
+
+
+	/**
+	* Translates the given uint8 RawValue to a float between 0.0 and +1.0, like most gamepad analog "triggers"
+	* 
+	* @param RawValue
+	* 
+	* @return			A float between 0.0 and +1.0
+	*/
+	constexpr float RawValueToFloatTrigger(const uint8 RawValue) const;
+
+	/**
+	* Translates the given uint8 RawValue to a float between -1.0 and +1.0.
+	* 
+	* @param RawValue
+	* @param Deadzone
+	* 
+	* @return A float between -1.0 and +1.0. Will be 0.0 if the value is within the deadzone.
+	*/
+	const float RawValueToFloatAnalog(uint8 RawValue, const uint8 DeadZone = 2) const;
+
+	/**
+	* Processes all the current raw input data. This should only be called after we have read our current raw input report
+	* or we have manually set the values (i.e. when we clear the input)
+	*/
+	bool ProcessAllRawValues(const FGameInputEventParams& Params, const FGameInputDeviceConfiguration* DeviceConfig, const bool bShouldProcessButtons, const bool bShouldProcessAnalog);
+
+	/**
+	* Processes the raw value at the given index as a button bitmask. 
+	* For every bit in the raw value and treat it as 0 for pressed and 1 for on. 
+	*/
+	bool ProcessRawInputValueAsBitmask(const FGameInputEventParams& Params, const int32 RawValueIndex, const FGameInputRawDeviceReportData* AxisData);
+
+	/**
+	* Processes the raw value at the given index as a button value. If the value is non-zero, then treat it as pressed. 
+	* If the value is zero, treat it as unpressed.
+	*/
+	bool ProcessRawInputValueAsButton(const FGameInputEventParams& Params, const int32 RawValueIndex, const FGameInputRawDeviceReportData* AxisData);
+
+	/**
+	* Processes the raw value at the given index as an analog value. This will translate the uint8 to a float value 
+	* and call OnControllerAnalog for it.
+	*/
+	bool ProcessRawInputValueAsAanalog(const FGameInputEventParams& Params, const int32 RawValueIndex, const FGameInputRawDeviceReportData* AxisData);
+
+	/**
+	* Processes a pair of raw input indexes. This allows you to combine the values of multiple uint8 axis readings and combine them
+	*/
+	bool ProcessRawInputValueAsAanalogPaired(const FGameInputEventParams& Params, const FGameInputRawDeviceReportData* AxisData);
+
+	/**
+	* The current raw input data that is populated by IGameInputRawDeviceReport::GetRawData
+	*/
+	TArray<uint8> CurrentRawData;
+
+	/**
+	* The previous frame's raw input data that is set at the end of ProcessInput.
+	*/
+	TArray<uint8> PreviousRawData;
+
+	/** A struct that contains info about each raw input index */
+	struct FPerRawInputIndexData
+	{
+		// The max number of buttons possibly associated with this raw input data. The raw input data is
+		// a uint8, so we can have at maximum a button per-bit, so 8.
+		enum { MaxSupportedButtons = 8 };
+
+		// Stores the times that each button index has been pressed so we can detect key repeats
+		double RepeatTime[MaxSupportedButtons];
+
+		/** 
+		* A map of key names to int values that is generated dynamically based on the button mapping.
+		*/
+		TMap<uint32, FName> KeyNameMap;
+	};
+
+	/**
+	* A map of an index to some raw data associated with it. 
+	* This is only used for button evaluation, not analog.
+	*/
+	TMap<int32, FPerRawInputIndexData> RawInputIndexDataMap;
+
+	/**
+	* Keeps track of how many times this gamepad has been processed this frame.
+	* Every successful processing of button input in ProcessInput will increment this value.
+	* This will get reset at the end of the input frame in PostProcessInput.
+	*/
+	int32 NumReadingsProcessedThisFrame = 0;
+};
+
+#endif	// GAME_INPUT_SUPPORT
\ No newline at end of file
diff --git a/Engine/Plugins/Runtime/GameInput/Source/GameInputBase/Public/GameInputKeyTypes.h b/Engine/Plugins/Runtime/GameInput/Source/GameInputBase/Public/GameInputKeyTypes.h
new file mode 100644
--- /dev/null
+++ b/Engine/Plugins/Runtime/GameInput/Source/GameInputBase/Public/GameInputKeyTypes.h
@@ -0,0 +1,72 @@
+// Copyright Epic Games, Inc. All Rights Reserved.
+
+#pragma once
+
+#include "InputCoreTypes.h"
+
+/**
+ * Key types that are unique to the Game Input plugin
+ * These keys are added on module startup of the Game Input module.
+ *
+ * HEY, LISTEN!
+ * If you add or otherwise modify these keys, make sure the change is reflected in FGameInputBaseModule::InitializeGameInputKeys.
+ */
+struct GAMEINPUTBASE_API FGameInputKeys
+{
+	//
+	// Racing Wheel
+	//
+
+	// Analog types
+	static const FKey RacingWheel_Brake;
+	static const FKey RacingWheel_Clutch;
+	static const FKey RacingWheel_Handbrake;
+	static const FKey RacingWheel_Throttle;
+	static const FKey RacingWheel_Wheel;
+	static const FKey RacingWheel_PatternShifterGear;
+
+	// Button types
+	static const FKey RacingWheel_None;
+	static const FKey RacingWheel_Menu;
+	static const FKey RacingWheel_View;
+	static const FKey RacingWheel_PreviousGear;
+	static const FKey RacingWheel_NextGear;
+
+	//
+	// Flight Stick
+	//
+
+	// Analog types
+	static const FKey FlightStick_Roll;
+	static const FKey FlightStick_Pitch;
+	static const FKey FlightStick_Yaw;
+	static const FKey FlightStick_Throttle;
+
+	// Button types
+	static const FKey FlightStick_None;
+	static const FKey FlightStick_Menu;
+	static const FKey FlightStick_View;
+	static const FKey FlightStick_FirePrimary;
+	static const FKey FlightStick_FireSecondary;
+
+	//
+	// Arcade Stick
+	//
+
+	// Button Types
+	static const FKey ArcadeStick_None;
+	static const FKey ArcadeStick_Menu;
+	static const FKey ArcadeStick_View;
+	static const FKey ArcadeStick_Up;
+	static const FKey ArcadeStick_Down;
+	static const FKey ArcadeStick_Left;
+	static const FKey ArcadeStick_Right;
+	static const FKey ArcadeStick_Action1;
+	static const FKey ArcadeStick_Action2;
+	static const FKey ArcadeStick_Action3;
+	static const FKey ArcadeStick_Action4;
+	static const FKey ArcadeStick_Action5;
+	static const FKey ArcadeStick_Action6;
+	static const FKey ArcadeStick_Special1;
+	static const FKey ArcadeStick_Special2;
+};
\ No newline at end of file
diff --git a/Engine/Plugins/Runtime/GameInput/Source/GameInputBase/Public/GameInputLogging.h b/Engine/Plugins/Runtime/GameInput/Source/GameInputBase/Public/GameInputLogging.h
new file mode 100644
--- /dev/null
+++ b/Engine/Plugins/Runtime/GameInput/Source/GameInputBase/Public/GameInputLogging.h
@@ -0,0 +1,7 @@
+// Copyright Epic Games, Inc. All Rights Reserved.
+
+#pragma once
+
+#include "Logging/LogMacros.h"
+
+GAMEINPUTBASE_API DECLARE_LOG_CATEGORY_EXTERN(LogGameInput, Log, All);
\ No newline at end of file
diff --git a/Engine/Plugins/Runtime/GameInput/Source/GameInputBase/Public/GameInputUtils.h b/Engine/Plugins/Runtime/GameInput/Source/GameInputBase/Public/GameInputUtils.h
new file mode 100644
--- /dev/null
+++ b/Engine/Plugins/Runtime/GameInput/Source/GameInputBase/Public/GameInputUtils.h
@@ -0,0 +1,78 @@
+// Copyright Epic Games, Inc. All Rights Reserved.
+
+#pragma once
+
+#if GAME_INPUT_SUPPORT
+
+#include "Containers/UnrealString.h"
+#include "GenericPlatform/GenericPlatformInputDeviceMapper.h"
+#include "GenericPlatform/GenericApplicationMessageHandler.h"
+
+THIRD_PARTY_INCLUDES_START
+#include "Microsoft/AllowMicrosoftPlatformTypes.h"
+#include <GameInput.h>
+#include "Microsoft/HideMicrosoftPlatformTypes.h"
+THIRD_PARTY_INCLUDES_END
+
+namespace UE::GameInput
+{
+	GAMEINPUTBASE_API FString LexToString(IGameInputDevice* Device);
+	
+	GAMEINPUTBASE_API FString LexToString(GameInputDeviceStatus DeviceStatus);
+
+	GAMEINPUTBASE_API FString LexToString(GameInputKind InputKind);
+
+	GAMEINPUTBASE_API FString LexToString(GameInputSwitchPosition SwitchPos);
+
+	GAMEINPUTBASE_API FString LexToString(const APP_LOCAL_DEVICE_ID& DeviceId);
+
+	GAMEINPUTBASE_API const TCHAR* GetMouseButtonName(EMouseButtons::Type MouseButton);
+
+	GAMEINPUTBASE_API EInputDeviceConnectionState DeviceStateToConnectionState(GameInputDeviceStatus InCurrentStatus, GameInputDeviceStatus InPreviousStatus);
+
+	/**
+	 * Populates the OutKeyName with the name of the Unreal Engine FKey that should be used for the given button mask, based on the
+	 * UE button map that should be used.
+	 *
+	 * @return true if a key name was found, false otherwise.
+	 */
+	GAMEINPUTBASE_API bool GameInputButtonToUnrealName(const TMap<uint32, FGamepadKeyNames::Type>& UEButtonMap, uint32 ButtonMask, OUT FGamepadKeyNames::Type& OutKeyName);
+
+	/** A map of game input switch position enums to Unreal Engine FKey names */
+	GAMEINPUTBASE_API const TMap<GameInputSwitchPosition, TArray<FGamepadKeyNames::Type>>& GetSwitchButtonMap();
+
+	/**
+	* Gets an array of FKeys that are associated with the given switch position.
+	* The array can have multiple keys if the GameInput switch state is "Up/Down + Left/Right"
+	*/
+	GAMEINPUTBASE_API const TArray<FGamepadKeyNames::Type>* SwitchPositionToUnrealName(const GameInputSwitchPosition ButtonMask);
+
+	// Button flags assigned by GameInput
+	enum { GamingInputButtonMask = 0x3FFF };
+
+	// Allocate flags for emulated/extended digital button functionality
+	enum EGameInputAuxButton
+	{
+		GamepadButtonAux_LeftTrigger = (1 << 22),
+		GamepadButtonAux_RightTrigger = (1 << 23),
+		GamepadButtonAux_LeftStickUp = (1 << 24),
+		GamepadButtonAux_LeftStickDown = (1 << 25),
+		GamepadButtonAux_LeftStickLeft = (1 << 26),
+		GamepadButtonAux_LeftStickRight = (1 << 27),
+		GamepadButtonAux_RightStickUp = (1 << 28),
+		GamepadButtonAux_RightStickDown = (1 << 29),
+		GamepadButtonAux_RightStickLeft = (1 << 30),
+		GamepadButtonAux_RightStickRight = (1 << 31),
+	};
+
+	// TODO_BH: Move these into a nice settings object instead of being hard coded
+	// Values taken from XINPUT_GAMEPAD_LEFT_THUMB_DEADZONE, XINPUT_GAMEPAD_RIGHT_THUMB_DEADZONE, XINPUT_GAMEPAD_TRIGGER_THRESHOLD used by XInput.
+	static constexpr float GamepadLeftStickDeadzone = (7849.0f / 32768.0f);
+	static constexpr float GamepadRightStickDeadzone = (8689.0f / 32768.0f);
+	static constexpr float GamepadTriggerDeadzone = (30.0f / 255.0f);
+
+	static constexpr float InitialRepeatDelay = (0.2f);
+	static constexpr float SubsequentRepeatDelay = (0.1f);
+}
+
+#endif	// GAME_INPUT_SUPPORT
\ No newline at end of file
diff --git a/Engine/Plugins/Runtime/GameInput/Source/GameInputBase/Public/IGameInputDeviceInterface.h b/Engine/Plugins/Runtime/GameInput/Source/GameInputBase/Public/IGameInputDeviceInterface.h
new file mode 100644
--- /dev/null
+++ b/Engine/Plugins/Runtime/GameInput/Source/GameInputBase/Public/IGameInputDeviceInterface.h
@@ -0,0 +1,165 @@
+// Copyright Epic Games, Inc. All Rights Reserved.
+
+#pragma once
+
+#include "IInputDevice.h"
+#include "Misc/CoreMiscDefines.h"	// For FInputDeviceId/FPlatformUserId
+
+#if GAME_INPUT_SUPPORT
+THIRD_PARTY_INCLUDES_START
+#include "Microsoft/AllowMicrosoftPlatformTypes.h"
+#include <GameInput.h>
+#include "Microsoft/HideMicrosoftPlatformTypes.h"
+THIRD_PARTY_INCLUDES_END
+#include "Microsoft/COMPointer.h"
+#endif	// GAME_INPUT_SUPPORT
+
+#include "GameInputDeviceContainer.h"
+
+/**
+* IInputDevice interface for the Game Input library. This is a base class
+* that is intended to be implemented on each platform that supports it (Xbox and PC).
+* 
+* This stores all state about every unique Human Interface Device that the GameInput library
+* detects. 
+* 
+* Any implementors of this interface should be Created in an IInputDevice module's CreateInputDevice function
+* (IInputDeviceModule::CreateInputDevice), and have it's Initialize function called
+*/
+class GAMEINPUTBASE_API IGameInputDeviceInterface : public IInputDevice
+{
+public:
+	explicit IGameInputDeviceInterface(const TSharedRef<FGenericApplicationMessageHandler>& InMessageHandler, struct IGameInput* InGameInput);
+	virtual ~IGameInputDeviceInterface() override;
+	
+	/** 
+	* Init function that should be called post-construction of this object. 
+	* This is required for sending controller events and will ensure if not called.
+	*/
+	void Initialize();
+
+protected:
+
+	//~ Begin IInputDevice interface
+	virtual void Tick(float DeltaTime) override;
+	virtual void SendControllerEvents() override;
+	virtual void SetMessageHandler(const TSharedRef<FGenericApplicationMessageHandler>& InMessageHandler) override;
+	virtual bool Exec(UWorld* InWorld, const TCHAR* Cmd, FOutputDevice& Ar) override;
+	virtual void SetChannelValue(int32 ControllerId, FForceFeedbackChannelType ChannelType, float Value) override;
+	virtual void SetChannelValues(int32 ControllerId, const FForceFeedbackValues& values) override;
+	virtual bool IsGamepadAttached() const override;
+	//~End IInputDevice interface
+
+#if GAME_INPUT_SUPPORT
+
+	/** Returns a mask of the currently supported Game Input Kind that we want to get the reading of. */
+	virtual GameInputKind GetCurrentGameInputKindSupport() const;
+
+	/** Binds a listener for device status changes from game input to this class */
+	virtual bool BindDeviceStatusCallbacks();
+	
+	/**
+	 * A callback for when an input device connection change has been detected.
+	 */
+	virtual void OnDeviceConnectionChanged(GameInputCallbackToken CallbackToken, IGameInputDevice* Device, uint64 Timestamp, GameInputDeviceStatus CurrentStatus, GameInputDeviceStatus PreviousStatus);
+
+	/**
+	 * Mark this device as disconnected and invalidate the tracked game input device. A disconnected device will stay in the
+	 * GameInputSourceDevices with a null Device pointer.
+	 */
+	virtual void HandleDeviceDisconnected(IGameInputDevice* Device, uint64 Timestamp) = 0;
+
+	/**
+	 * Map this new input device to a platform user and initialize the known FGameInputDeviceInfo.
+	 * If this is a new connection, this will request a new FInputDeviceId from the platform input device mapper
+	 * and add it to the GameInputSourceDevices map.
+	 */
+	virtual void HandleDeviceConnected(IGameInputDevice* Device, uint64 Timestamp) = 0;
+
+	/** Process any deferred device connection changes that may have occurred. This should only be called on the game thread. */
+	void ProcessDeferredDeviceConnectionChanges();
+	
+	/** Callbacks for when the app is constrained/unconstrained (i.e. you Alt + Tab out of the window, or open the home menu on a console platform) */
+	void OnAppConstrained();
+	void OnAppUnconstrained();
+	
+	/** Called on the first update after the was un-constrained. */
+	void DetermineStateAfterFirstUnconstrainedUpdate();
+
+	/** Returns a matching FGameInputDeviceContainer if one exists for the given GameInput device. Nullptr if none exist. */
+	FGameInputDeviceContainer* GetDeviceData(IGameInputDevice* InDevice);
+
+	/**
+	* Returns a pointer to the Device info associated with the given IGameInputDevice.
+	*
+	* If one doesn't exist, then create a new one and return that.
+	*/
+	FGameInputDeviceContainer* GetOrCreateDeviceData(IGameInputDevice* InDevice);
+
+	/**
+	* Creates a new FGameInputDeviceContainer* and stores it in DeviceData. It is up to each implementing 
+	* device interface to how that device is mapped to platform users and initialized.
+	*/
+	virtual FGameInputDeviceContainer* CreateDeviceData(IGameInputDevice* InDevice) = 0;
+
+	/**
+	* This will update the CurrentlyConnectedDeviceTypes bitmask based on the current set of DeviceData
+	*/
+	virtual void EnumerateCurrentlyConnectedDeviceTypes();
+	
+	/** Message handler that we can use to tell the engine about input events */
+	TSharedRef<FGenericApplicationMessageHandler> MessageHandler;
+	
+	/** Pointer to the Game Input interface object */
+	TComPtr<IGameInput> GameInput;
+
+	// CS for OnDeviceConnectionChanged called from background system thread and PollGameDeviceState called from Main thread
+	FCriticalSection DeviceInfoCS;
+
+	/** Token for when a device connection change happens */
+	GameInputCallbackToken ConnectionChangeCallbackToken;
+	
+	/** An enumeration of all the currently registered Game input Devices. */
+	GameInputKind CurrentlyConnectedDeviceTypes;
+
+	/** Set to true if the IGameInputDeviceInterface::Initialize function was called. This is required for sending controller events. */
+	bool bWasinitialized;
+
+	/** True if the app is currently constrained this tick */
+	bool bIsAppCurrentlyConstrained;
+
+	/** True if the app was constrained last tick */
+	bool bWasAppConstrainedLastTick;
+
+	// Array of Input Device data for every connected Game Input Device.
+	TArray<TSharedPtr<FGameInputDeviceContainer>> DeviceData;
+
+	// A mapping of a platform user to the most recently used Input Device for that user
+	TMap<FPlatformUserId, TSharedPtr<FGameInputDeviceContainer>> PlatformUserIdToMostRecentDeviceContainer;	
+
+	/** Stores some data about when device connection changes happen to be processed later */
+	struct FDeferredDeviceConnectionChanges
+	{
+		FDeferredDeviceConnectionChanges() = default;
+
+		/** The GameInput device that has changed */
+		IGameInputDevice* Device = nullptr;
+		
+		/** Timestamp of when the connection state change happened */
+		uint64 Timestamp = 0;
+		
+		/** The evaluated status of the game input device */
+		EInputDeviceConnectionState Status = EInputDeviceConnectionState::Invalid;
+		
+		/** The current GameInput status of the device */
+		GameInputDeviceStatus CurrentStatus = GameInputDeviceStatus::GameInputDeviceNoStatus;
+
+		/** The previous GameInput status of the device */
+		GameInputDeviceStatus PreviousStatus = GameInputDeviceStatus::GameInputDeviceNoStatus;
+	};
+
+	/** Array of device connection changed events that will be processed in ProcessDeferredDeviceConnectionChanges. */
+	TArray<FDeferredDeviceConnectionChanges> DeferredDeviceConnectionChanges;
+
+#endif	// #if GAME_INPUT_SUPPORT
+};
\ No newline at end of file
diff --git a/Engine/Plugins/Runtime/GameInput/Source/GameInputBaseEditor/GameInputBaseEditor.build.cs b/Engine/Plugins/Runtime/GameInput/Source/GameInputBaseEditor/GameInputBaseEditor.build.cs
new file mode 100644
--- /dev/null
+++ b/Engine/Plugins/Runtime/GameInput/Source/GameInputBaseEditor/GameInputBaseEditor.build.cs
@@ -0,0 +1,19 @@
+// Copyright Epic Games, Inc. All Rights Reserved.
+
+using UnrealBuildTool;
+
+public class GameInputBaseEditor : ModuleRules
+{
+	public GameInputBaseEditor(ReadOnlyTargetRules Target) : base(Target)
+	{
+		PrivateDependencyModuleNames.AddRange(
+			new string[] {
+				"Core",
+				"CoreUObject",
+				"Engine",
+				"GameInputBase",
+				"Settings"
+			}
+		);
+	}
+}
diff --git a/Engine/Plugins/Runtime/GameInput/Source/GameInputBaseEditor/Private/GameInputBaseEditorModule.cpp b/Engine/Plugins/Runtime/GameInput/Source/GameInputBaseEditor/Private/GameInputBaseEditorModule.cpp
new file mode 100644
--- /dev/null
+++ b/Engine/Plugins/Runtime/GameInput/Source/GameInputBaseEditor/Private/GameInputBaseEditorModule.cpp
@@ -0,0 +1,41 @@
+// Copyright Epic Games, Inc. All Rights Reserved.
+
+#include "ISettingsModule.h"
+#include "Modules/ModuleInterface.h"
+#include "Modules/ModuleManager.h"
+#include "GameInputDeveloperSettings.h"
+
+#define LOCTEXT_NAMESPACE "GameInputBaseEditor"
+
+/**
+* Editor module for Game Input that will register the UGameInputDeveloperSettings so they
+* show up as editable project settings in the editor.
+*/
+class FGameInputBaseEditorModule final : public IModuleInterface
+{
+private:
+	virtual void StartupModule() override
+	{
+		// register settings
+		if (ISettingsModule* SettingsModule = FModuleManager::GetModulePtr<ISettingsModule>("Settings"))
+		{
+			SettingsModule->RegisterSettings("Project", "Plugins", "Input",
+				LOCTEXT("ControllerSettingsName", "Game Input Device Settings"),
+				LOCTEXT("ControllerSettingsDescription", "Settings for the Game Input plugin."),
+				GetMutableDefault<UGameInputDeveloperSettings>()
+			);
+		}
+	}
+
+	virtual void ShutdownModule() override
+	{	
+		if (ISettingsModule* SettingsModule = FModuleManager::GetModulePtr<ISettingsModule>("Settings"))
+		{
+			SettingsModule->UnregisterSettings("Project", "Plugins", "Input");
+		}
+	}
+};
+
+IMPLEMENT_MODULE(FGameInputBaseEditorModule, GameInputBaseEditor);
+
+#undef LOCTEXT_NAMESPACE
\ No newline at end of file
