Change the ToString for a game input device to include the decimal representation. This is useful for if you are setting up your own device configuration.

#rnx

diff --git a/Engine/Plugins/Runtime/GameInput/Source/GameInputBase/Private/GameInputUtils.cpp b/Engine/Plugins/Runtime/GameInput/Source/GameInputBase/Private/GameInputUtils.cpp
--- a/Engine/Plugins/Runtime/GameInput/Source/GameInputBase/Private/GameInputUtils.cpp
+++ b/Engine/Plugins/Runtime/GameInput/Source/GameInputBase/Private/GameInputUtils.cpp
@@ -15,12 +15,12 @@ namespace UE::GameInput
 			const GameInputDeviceInfo* Info = Device->GetDeviceInfo();
 			const bool bIsVirtualDevice = Info->deviceFamily == GameInputFamilyVirtual;
 			
-			Result = FString::Printf(TEXT("%llp   DeviceName: %s %s ProdId: %04x VendId: %04x"), 
+			Result = FString::Printf(TEXT("%llp   DeviceName: %s %s ProdId: %04x (%u) VendId: %04x (%u)"), 
 				Device, 
 				Info->displayName, 
 				bIsVirtualDevice ? TEXT("(virtual)") : TEXT(""),
-				Info->productId,
-				Info->vendorId);
+				Info->productId, (uint32)Info->productId,
+				Info->vendorId, (uint32)Info->vendorId);
 		}
 
 		return Result;
