Experimental Racing Wheel support for Game Input Windows!

diff --git a/Engine/Plugins/Runtime/GameInput/Source/GameInputBase/Private/GameInputDeviceContainer.cpp b/Engine/Plugins/Runtime/GameInput/Source/GameInputBase/Private/GameInputDeviceContainer.cpp
--- a/Engine/Plugins/Runtime/GameInput/Source/GameInputBase/Private/GameInputDeviceContainer.cpp
+++ b/Engine/Plugins/Runtime/GameInput/Source/GameInputBase/Private/GameInputDeviceContainer.cpp
@@ -141,17 +141,13 @@ void FGameInputDeviceContainer::InitalizeDeviceProcessors_Impl()
 		}
 	}
 
-	// TODO: Implement these kinds of input! 
-	
-	//if (DeviceKind & AllowedGameInputKinds &  GameInputKindController)
-	//{
-
-	//}
-
-	//if (DeviceKind & AllowedGameInputKinds & GameInputKindRacingWheel)
-	//{
+	// Racing wheels
+	if (DeviceKind & AllowedGameInputKinds & GameInputKindRacingWheel)
+	{
+		Processors.Emplace(MakeShared<FGameInputRacingWheelProcessor>());
+	}
 
-	//}
+	// TODO: Implement these kinds of input! 
 
 	//if (DeviceKind & AllowedGameInputKinds & GameInputKindFlightStick)
 	//{
@@ -162,11 +158,6 @@ void FGameInputDeviceContainer::InitalizeDeviceProcessors_Impl()
 	//{
 
 	//}
-
-	//if (DeviceKind &AllowedGameInputKinds &  GameInputKindRawDeviceReport)
-	//{
-
-	//}
 }
 
 const GameInputKind FGameInputDeviceContainer::ProcessInput(IGameInput* GameInput, const GameInputKind CurrentSupportedKind, const GameInputKind ProcessedKindsForPlatformUserThisFrame)
diff --git a/Engine/Plugins/Runtime/GameInput/Source/GameInputBase/Private/GameInputDeviceProcessor.cpp b/Engine/Plugins/Runtime/GameInput/Source/GameInputBase/Private/GameInputDeviceProcessor.cpp
--- a/Engine/Plugins/Runtime/GameInput/Source/GameInputBase/Private/GameInputDeviceProcessor.cpp
+++ b/Engine/Plugins/Runtime/GameInput/Source/GameInputBase/Private/GameInputDeviceProcessor.cpp
@@ -17,6 +17,27 @@
 namespace UE::GameInput
 {
 	static FName InputClassName = FName("GameInput");
+
+	/**
+	* Returns a map of GameInputRacingWheelButtons to their associated Unreal Engine FKey names.
+	*/
+	static const TMap<uint32, FGamepadKeyNames::Type>& GetRacingWheelButtonMap()
+	{
+		static const TMap<uint32, FGamepadKeyNames::Type> RacingWheelButtonMap
+		{
+			{ static_cast<uint32>(GameInputRacingWheelButtons::GameInputRacingWheelNone), FGameInputKeys::RacingWheel_None.GetFName() },
+			{ static_cast<uint32>(GameInputRacingWheelButtons::GameInputRacingWheelMenu), FGameInputKeys::RacingWheel_Menu.GetFName() },
+			{ static_cast<uint32>(GameInputRacingWheelButtons::GameInputRacingWheelView), FGameInputKeys::RacingWheel_View.GetFName() },
+			{ static_cast<uint32>(GameInputRacingWheelButtons::GameInputRacingWheelPreviousGear), FGameInputKeys::RacingWheel_PreviousGear.GetFName() },
+			{ static_cast<uint32>(GameInputRacingWheelButtons::GameInputRacingWheelNextGear), FGameInputKeys::RacingWheel_NextGear.GetFName() },
+			{ static_cast<uint32>(GameInputRacingWheelButtons::GameInputRacingWheelDpadUp), FGamepadKeyNames::DPadUp },
+			{ static_cast<uint32>(GameInputRacingWheelButtons::GameInputRacingWheelDpadDown), FGamepadKeyNames::DPadDown },
+			{ static_cast<uint32>(GameInputRacingWheelButtons::GameInputRacingWheelDpadLeft), FGamepadKeyNames::DPadLeft },
+			{ static_cast<uint32>(GameInputRacingWheelButtons::GameInputRacingWheelDpadRight), FGamepadKeyNames::DPadRight },
+		};
+
+		return RacingWheelButtonMap;
+	}
 };
 
 ////////////////////////////////////////////////////////
@@ -679,7 +700,7 @@ FGameInputControllerDeviceProcessor::FGameInputControllerDeviceProcessor()
 {
 	// Ensure that our switch repeat times array has some default values set on it by default
 	// because we access it with the [] operator when processing it. 
-	SwitchRepeatTimes.AddDefaulted(static_cast<uint32>(GameInputSwitchLeft) + 1);
+	SwitchRepeatTimes.AddDefaulted(static_cast<uint32>(GameInputSwitchUpLeft) + 1);
 }
 
 bool FGameInputControllerDeviceProcessor::ProcessInput(const FGameInputEventParams& Params)
@@ -1755,4 +1776,186 @@ GameInputKind FGameInputRawDeviceProcessor::GetSupportedReadingKind() const
 	return GameInputKindRawDeviceReport;
 }
 
+///////////////////////////////////////////////////////////////////////////////////////
+// FGameInputRacingWheelProcessor
+
+FGameInputRacingWheelProcessor::FGameInputRacingWheelProcessor()
+	: IGameInputDeviceProcessor()
+{
+	NumReadingsProcessedThisFrame = 0;
+	FMemory::Memset(PreviousState, 0);
+	FMemory::Memset(RepeatTime, 0);
+}
+
+bool FGameInputRacingWheelProcessor::ProcessInput(const FGameInputEventParams& Params)
+{
+	bool bRes = false;
+
+	// Can't do anything for an invalid platform user
+	if (!Params.PlatformUserId.IsValid() || !Params.Reading)
+	{
+		return bRes;
+	}
+
+	GameInputRacingWheelState WheelState;
+	if (!Params.Reading->GetRacingWheelState(&WheelState))
+	{
+		return bRes;
+	}
+
+	// We only want to process the buttons here, as it might get called multiple times per frame.
+	bRes |= ProcessWheelButtonState(Params, WheelState);
+
+	++NumReadingsProcessedThisFrame;
+
+	return bRes;
+}
+
+bool FGameInputRacingWheelProcessor::PostProcessInput(const FGameInputEventParams& Params)
+{
+	bool bRes = false;
+
+	// Check if we have already processed buttons this frame. If we haven't we want to do it
+	const bool bHasProcessedAnyButtonsThisFrame = NumReadingsProcessedThisFrame > 0;	
+	NumReadingsProcessedThisFrame = 0;
+
+	// Can't do anything for an invalid platform user
+	if (!Params.PlatformUserId.IsValid() || !Params.PreviousReading)
+	{
+		return bRes;
+	}
+
+	// Use the "PreviousRading" because we only want to process the analog inputs once, and this will
+	// point to the most up to date reading.
+	GameInputRacingWheelState WheelState;
+	if (!Params.PreviousReading->GetRacingWheelState(&WheelState))
+	{
+		return bRes;
+	}
+
+	if (!bHasProcessedAnyButtonsThisFrame)
+	{
+		bRes |= ProcessWheelButtonState(Params, WheelState);
+	}
+
+	bRes |= ProcessWheelAnalogState(Params, WheelState);
+
+	return bRes;
+}
+
+void FGameInputRacingWheelProcessor::ClearState(const FGameInputEventParams& Params)
+{
+	// Can't do anything for an invalid platform user
+	if (!Params.PlatformUserId.IsValid())
+	{
+		return;
+	}
+
+	// We can simply process a wheel state where everything is zero
+	GameInputRacingWheelState ZeroValueState = {};
+	FMemory::Memset(ZeroValueState, 0);
+
+	ProcessWheelButtonState(Params, ZeroValueState);
+	ProcessWheelAnalogState(Params, ZeroValueState);
+
+	// Zero out any state trackers
+	NumReadingsProcessedThisFrame = 0;
+	FMemory::Memset(PreviousState, 0);
+	FMemory::Memset(RepeatTime, 0);
+}
+
+namespace UE::GameInput
+{
+	static bool HasDifferentWheelAnalogInput(const GameInputRacingWheelState& CurrentState, const GameInputRacingWheelState& PreviousState)
+	{
+		// If anything differs from the previous reading, then it has different inputs
+		return 
+			CurrentState.wheel != PreviousState.wheel ||
+			CurrentState.throttle != PreviousState.throttle ||
+			CurrentState.brake != PreviousState.brake ||
+			CurrentState.clutch != PreviousState.clutch ||
+			CurrentState.handbrake != PreviousState.handbrake;
+	}
+};
+
+const float FGameInputRacingWheelProcessor::GetRacingWheelDeadzone()
+{
+	// These settings should always be available
+	if (const UGameInputPlatformSettings* Settings = UGameInputPlatformSettings::Get())
+	{
+		return Settings->RacingWheelDeadzone;
+	}
+
+	return UGameInputPlatformSettings::DefaultRacingWheelDeadzone;
+}
+
+bool FGameInputRacingWheelProcessor::ProcessWheelAnalogState(const FGameInputEventParams& Params, GameInputRacingWheelState& CurrentWheelState)
+{
+	// ignore this input if the reading has remained empty from last time
+	const float Deadzone = GetRacingWheelDeadzone();
+
+	// If the analog values haven't changed, don't bother sending any events for them
+	const bool bHasDifferentAnalogInput = UE::GameInput::HasDifferentWheelAnalogInput(CurrentWheelState, PreviousState);
+	if (!bHasDifferentAnalogInput)
+	{
+		return false;
+	}
+
+	OnControllerAnalog(Params, FGameInputKeys::RacingWheel_Brake.GetFName(), CurrentWheelState.brake, PreviousState.brake, Deadzone);
+	OnControllerAnalog(Params, FGameInputKeys::RacingWheel_Clutch.GetFName(), CurrentWheelState.clutch, PreviousState.clutch, Deadzone);
+	OnControllerAnalog(Params, FGameInputKeys::RacingWheel_Handbrake.GetFName(), CurrentWheelState.handbrake, PreviousState.handbrake, Deadzone);
+	OnControllerAnalog(Params, FGameInputKeys::RacingWheel_Throttle.GetFName(), CurrentWheelState.throttle, PreviousState.throttle, Deadzone);
+	OnControllerAnalog(Params, FGameInputKeys::RacingWheel_Wheel.GetFName(), CurrentWheelState.wheel, PreviousState.wheel, Deadzone);
+
+	// Do we actually want this as a float? We should test the values that this can produce
+	OnControllerAnalog(Params, 
+		FGameInputKeys::RacingWheel_PatternShifterGear.GetFName(), 
+		static_cast<float>(CurrentWheelState.patternShifterGear), 
+		static_cast<float>(PreviousState.patternShifterGear), 
+		Deadzone);
+
+	// Keep track of the previous state
+	PreviousState.brake = CurrentWheelState.brake;
+	PreviousState.clutch = CurrentWheelState.clutch;
+	PreviousState.handbrake = CurrentWheelState.handbrake;
+	PreviousState.throttle = CurrentWheelState.throttle;
+	PreviousState.wheel = CurrentWheelState.wheel;
+	PreviousState.patternShifterGear = CurrentWheelState.patternShifterGear;
+
+	return true;
+}
+
+bool FGameInputRacingWheelProcessor::ProcessWheelButtonState(const FGameInputEventParams& Params, GameInputRacingWheelState& CurrentWheelState)
+{
+	// If there has been no buttons pressed on this state or the previous one, don't bother trying
+	// to evaluate any events.
+	if (CurrentWheelState.buttons == 0 && PreviousState.buttons	== 0)
+	{
+		return false;
+	}
+
+	// This might not be necessary if the racing wheels also show up as "gamepad" devices...
+	const uint32 CurrentButtonHeldMask = static_cast<uint32>(CurrentWheelState.buttons);
+
+	uint32 LastButtonHeldMask = static_cast<uint32>(PreviousState.buttons);
+
+	EvaluateButtonStates(
+		Params,
+		CurrentButtonHeldMask,
+		OUT LastButtonHeldMask,
+		RepeatTime,
+		UE::GameInput::GetRacingWheelButtonMap(),
+		MaxSupportedButtons);
+
+	// Update the previous state here
+	PreviousState.buttons = CurrentWheelState.buttons;
+
+	return true;
+}
+
+GameInputKind FGameInputRacingWheelProcessor::GetSupportedReadingKind() const
+{
+	return GameInputKindRacingWheel;
+}
+
 #endif	// GAME_INPUT_SUPPORT
\ No newline at end of file
diff --git a/Engine/Plugins/Runtime/GameInput/Source/GameInputBase/Private/IGameInputDeviceInterface.cpp b/Engine/Plugins/Runtime/GameInput/Source/GameInputBase/Private/IGameInputDeviceInterface.cpp
--- a/Engine/Plugins/Runtime/GameInput/Source/GameInputBase/Private/IGameInputDeviceInterface.cpp
+++ b/Engine/Plugins/Runtime/GameInput/Source/GameInputBase/Private/IGameInputDeviceInterface.cpp
@@ -281,7 +281,12 @@ GameInputKind IGameInputDeviceInterface::GetCurrentGameInputKindSupport() const
 		RegisterInputKindMask |= GameInputKindMouse;
 	}
 
-	// TODO_BH: Future expansion of game input on Xbox
+	if (PlatformSettings->bProcessRacingWheel)
+	{
+		RegisterInputKindMask |= GameInputKindRacingWheel;
+	}
+
+	// TODO: Future expansion of GameInput devices!
 	/*if (Settings->bProcessArcadeStick)
 	{
 
diff --git a/Engine/Plugins/Runtime/GameInput/Source/GameInputBase/Public/GameInputDeveloperSettings.h b/Engine/Plugins/Runtime/GameInput/Source/GameInputBase/Public/GameInputDeveloperSettings.h
--- a/Engine/Plugins/Runtime/GameInput/Source/GameInputBase/Public/GameInputDeveloperSettings.h
+++ b/Engine/Plugins/Runtime/GameInput/Source/GameInputBase/Public/GameInputDeveloperSettings.h
@@ -339,6 +339,31 @@ public:
 	*/
 	UPROPERTY(Config, EditAnywhere, Category="Processing Options")
 	bool bProcessMouse = true;
+
+	/**
+	 * If true, then we will process the GameInputKindRacingWheel type.
+	 * 
+	 * Racing Wheels often times need to be used in conjunction with the "Controller" or "Gamepad" processors as well
+	 * in order to handle any "normal" buttons on them. For example, you may have a racing wheel with some ABXY buttons on it
+	 * that you want to process as well.
+	 * 
+	 * Note: This is experimental!
+	 * 
+	 * Default: False
+	 */
+	UPROPERTY(EditAnywhere, Config, Category = "Processing Options", meta=(DisplayName="Process Racing Wheel (Experimental)"))
+	bool bProcessRacingWheel = false;
+
+	/** The default racing wheel deadzone */
+	static constexpr float DefaultRacingWheelDeadzone = (7849.0f / 32768.0f);
+
+	/**
+	* The deadzone that should be applied when processing Racing Wheel analog values.
+	* 
+	* @see FGameInputRacingWheelProcessor::ProcessWheelAnalogState
+	*/
+	UPROPERTY(EditAnywhere, Config, Category = "Device Settings|Racing Wheel", meta=(EditCondition="bProcessRacingWheel"))
+	float RacingWheelDeadzone = DefaultRacingWheelDeadzone;
 };
 
 /**
diff --git a/Engine/Plugins/Runtime/GameInput/Source/GameInputBase/Public/GameInputDeviceProcessor.h b/Engine/Plugins/Runtime/GameInput/Source/GameInputBase/Public/GameInputDeviceProcessor.h
--- a/Engine/Plugins/Runtime/GameInput/Source/GameInputBase/Public/GameInputDeviceProcessor.h
+++ b/Engine/Plugins/Runtime/GameInput/Source/GameInputBase/Public/GameInputDeviceProcessor.h
@@ -446,4 +446,53 @@ protected:
 	int32 NumReadingsProcessedThisFrame = 0;
 };
 
+/**
+* Processor for racing wheel types (GameInputKindRacingWheel)
+* 
+* Note: Racing wheels are often paired with other device processors like Gamepad/Controller
+* to handle any "normal" buttons that may be on them.
+*/
+class GAMEINPUTBASE_API FGameInputRacingWheelProcessor : public IGameInputDeviceProcessor
+{
+public:
+	FGameInputRacingWheelProcessor();
+protected:
+	virtual bool ProcessInput(const FGameInputEventParams& Params) override;
+	virtual bool PostProcessInput(const FGameInputEventParams& Params) override;
+	virtual void ClearState(const FGameInputEventParams& Params) override;
+	virtual GameInputKind GetSupportedReadingKind() const override;
+
+	/**
+	* Returns the deadzone that we would like to use for racing wheel analog inputs.
+	*/
+	static const float GetRacingWheelDeadzone();
+
+	/**
+	* Processes the analog inputs of the wheel state (brake/clutch pedals, wheel movement, etc).
+	*/
+	bool ProcessWheelAnalogState(const FGameInputEventParams& Params, GameInputRacingWheelState& CurrentWheelState);
+
+	/**
+	* Processes any buttons on the racing wheel (next/previous gears, menu nav, etc);
+	*/
+	bool ProcessWheelButtonState(const FGameInputEventParams& Params, GameInputRacingWheelState& CurrentWheelState);
+
+	/**
+	* Keeps track of how many times this gamepad has been processed this frame.
+	* Every successful processing of button input in ProcessInput will increment this value.
+	* This will get reset at the end of the input frame in PostProcessInput.
+	*/
+	int32 NumReadingsProcessedThisFrame = 0;
+
+	/** The previous game input state that was last processed. */
+	GameInputRacingWheelState PreviousState;
+
+	/** 
+	* Array of repeat times to calculate if a button has been held long enough
+	* to receive an IE_REPEAT event.
+	*/
+	static const uint32 MaxSupportedButtons = 16;
+	double RepeatTime[MaxSupportedButtons];
+};
+
 #endif	// GAME_INPUT_SUPPORT
\ No newline at end of file
