Check for the GameInput.dll existing on in the System Library on desktop platforms before attempting to call any game input functions. If you run GameInput on a platform like Windows Server, or something else that has stripped down System32 DLL's, then it would crash attempting to load. Checking this allows for a graceful handling of the issue

We can actually remove the DelayLoaded.dll and GameInput.lib files here as well, because we are manually loading the DLL.

#rnx

diff --git a/Engine/Plugins/Runtime/GameInput/Source/GameInputBase/GameInputBase.build.cs b/Engine/Plugins/Runtime/GameInput/Source/GameInputBase/GameInputBase.build.cs
--- a/Engine/Plugins/Runtime/GameInput/Source/GameInputBase/GameInputBase.build.cs
+++ b/Engine/Plugins/Runtime/GameInput/Source/GameInputBase/GameInputBase.build.cs
@@ -58,16 +58,10 @@ namespace UnrealBuildTool.Rules
 			string GRDKLatestPath = Environment.GetEnvironmentVariable("GRDKLatest");
 			
 			// GameInput.h is from    <GRDKLatest>\GameKit\Include\GameInput.h
-			// GameInput.lib is from  <GRDKLatest>\GameKit\Lib\amd64\GameInput.lib
 			
-			string LibPath = Path.Combine(GRDKLatestPath, "GameKit", "Lib", "amd64", "GameInput.lib");
 			string IncludePath = Path.Combine(GRDKLatestPath, "GameKit", "Include");
 			
-			PublicAdditionalLibraries.Add(LibPath);
 			PublicSystemIncludePaths.Add(IncludePath);
-			
-			// Add this as a delay loaded DLL for cooking of the client
-			PublicDelayLoadDLLs.Add("GameInput.dll");
 		}
 
 		public GameInputBase(ReadOnlyTargetRules Target) : base(Target)
diff --git a/Engine/Plugins/Runtime/GameInput/Source/GameInputBase/Private/GameInputBaseModule.cpp b/Engine/Plugins/Runtime/GameInput/Source/GameInputBase/Private/GameInputBaseModule.cpp
--- a/Engine/Plugins/Runtime/GameInput/Source/GameInputBase/Private/GameInputBaseModule.cpp
+++ b/Engine/Plugins/Runtime/GameInput/Source/GameInputBase/Private/GameInputBaseModule.cpp
@@ -6,6 +6,7 @@
 #include "GameInputBaseIncludes.h"
 #include "GameInputKeyTypes.h"
 #include "GameInputLogging.h"
+#include "HAL/PlatformProcess.h"
 #include "Misc/App.h"
 #include "Modules/ModuleManager.h"
 
@@ -26,6 +27,12 @@ namespace UE::GameInput
 	// A singleton pointer to the base GameInput interface. 
 	// This provides access to reading the input stream, device callbacks, and more. 
 	static TComPtr<IGameInput> GGameInputInterface;
+
+#if PLATFORM_WINDOWS
+	// The name of the Game Input DLL file which would be needed when running on windows
+	static const FString GameInputDLLPath = TEXT("GameInput.dll");
+#endif	// #if PLATFORM_WINDOWS
+
 #endif	// #if GAME_INPUT_SUPPORT	
 }
 
@@ -73,8 +80,39 @@ void FGameInputBaseModule::StartupModule()
 
 #if GAME_INPUT_SUPPORT
 
+#if PLATFORM_WINDOWS
+	// Check to see if the GameInput.dll exists...
+	// Search for the GameInput dll on desktop platforms. If for some reason it doesn't exist, then we shouldn't
+	// attempt to call any functions from it. The only known case for this is when running a client on a server OS
+	// which doesn't have game input installed by default. 
+	GameInputDLLHandle = FPlatformProcess::GetDllHandle(*UE::GameInput::GameInputDLLPath);
+	if (GameInputDLLHandle == nullptr)
+	{
+		UE_LOG(LogGameInput, Warning, TEXT("[%hs] module is exiting because '%s' cannot be found. GameInput will not be initalized. Is it installed correctly?"), __func__, *UE::GameInput::GameInputDLLPath);
+		return;
+	}
+
+// Disable warning C4191: 'type cast' : unsafe conversion from 'PROC' to 'XXX' getting the GameInputCreate function
+#pragma warning(push)
+#pragma warning(disable:4191)
+	
+	// Create the Game Input interface
+	typedef HRESULT(WINAPI* FGameInputCreateFn)(IGameInput**);
+	FGameInputCreateFn pfnGameInputCreate = (FGameInputCreateFn)GetProcAddress((HMODULE)GameInputDLLHandle, "GameInputCreate");
+
+	UE_CLOG(!pfnGameInputCreate, LogGameInput, Warning, TEXT("[%hs] Failed to GetProcAddress (GameInputCreate). Game Input will fail to be created."), __func__);
+	
+	const HRESULT HResult = pfnGameInputCreate ? pfnGameInputCreate(&UE::GameInput::GGameInputInterface) : E_FAIL;
+
+#pragma warning(pop)
+
+#else 	// PLATFORM_WINDOWS
+	
 	// Create the Game Input interface
 	const HRESULT HResult = GameInputCreate(&UE::GameInput::GGameInputInterface);
+	
+#endif	// !PLATFORM_WINDOWS
+
 	UE_CLOG(FAILED(HResult), LogGameInput, Warning, TEXT("Failed to initialize GameInput: 0x%X"), HResult);
 
 	if (SUCCEEDED(HResult))
@@ -92,6 +130,14 @@ void FGameInputBaseModule::ShutdownModule()
 {
 #if GAME_INPUT_SUPPORT
 	UE::GameInput::GGameInputInterface.Reset();
+
+#if PLATFORM_WINDOWS
+	if (GameInputDLLHandle)
+	{
+		FPlatformProcess::FreeDllHandle(GameInputDLLHandle);
+	}
+#endif // #if PLATFORM_WINDOWS
+
 #endif // #if GAME_INPUT_SUPPORT
 }
 
diff --git a/Engine/Plugins/Runtime/GameInput/Source/GameInputBase/Public/GameInputBaseModule.h b/Engine/Plugins/Runtime/GameInput/Source/GameInputBase/Public/GameInputBaseModule.h
--- a/Engine/Plugins/Runtime/GameInput/Source/GameInputBase/Public/GameInputBaseModule.h
+++ b/Engine/Plugins/Runtime/GameInput/Source/GameInputBase/Public/GameInputBaseModule.h
@@ -32,4 +32,10 @@ public:
 protected:
 
 	void InitializeGameInputKeys();
+
+#if PLATFORM_WINDOWS && GAME_INPUT_SUPPORT
+	// Handle to the game input dll which is set on StartupModule.
+	// If we can't find the DLL then we will early exit and not attempt to initalize GameInput.
+	void* GameInputDLLHandle = nullptr;
+#endif // endif PLATFORM_WINDOWS && GAME_INPUT_SUPPORT
 };
\ No newline at end of file
