Add the ConfigRestartRequired metadata tag to Game Input settings. When you toggle these options in the editor you need to reboot to have the device connection logic run again and re-evaluate what device processors should be created, so this makes for a better UX.

Also collapse the bProcess Raw Input / Controller options into a category called "Require Configuration", because in order to process any of these device types you need to specific some kind of button mappings for them.

#rnx

diff --git a/Engine/Plugins/Runtime/GameInput/Source/GameInputBase/Public/GameInputDeveloperSettings.h b/Engine/Plugins/Runtime/GameInput/Source/GameInputBase/Public/GameInputDeveloperSettings.h
--- a/Engine/Plugins/Runtime/GameInput/Source/GameInputBase/Public/GameInputDeveloperSettings.h
+++ b/Engine/Plugins/Runtime/GameInput/Source/GameInputBase/Public/GameInputDeveloperSettings.h
@@ -232,7 +232,7 @@ struct GAMEINPUTBASE_API FGameInputDeviceConfiguration
 	/**
 	* A map of uint32 button index to an associated FName Unreal gamepad key name.
 	* 
-	* These key values should be ([button index] << 1) i.e, powers of two.
+	* These key values should be (1 << [button index]) i.e, powers of two.
 	*/
 	UPROPERTY(EditAnywhere, Config, Category = "Device Settings|Controller", meta=(EditCondition="bProcessControllerButtons"))
 	TMap<uint32, FName> ControllerButtonMappingData;
@@ -290,7 +290,7 @@ public:
 	 *
 	 * Default: False
 	 */
-	UPROPERTY(Config, EditAnywhere, Category="Processing Options")
+	UPROPERTY(Config, EditAnywhere, Category="Processing Options|Require Configuration", meta = (ConfigRestartRequired = true))
 	bool bProcessController = false;
 
 	/** 
@@ -298,7 +298,7 @@ public:
 	*
 	* Default: False
 	*/
-	UPROPERTY(Config, EditAnywhere, Category="Processing Options")
+	UPROPERTY(Config, EditAnywhere, Category="Processing Options|Require Configuration", meta = (ConfigRestartRequired = true))
 	bool bProcessRawInput = false;
 	
 	/**
@@ -307,7 +307,7 @@ public:
 	* 
 	* Default: True
 	*/
-	UPROPERTY(Config, EditAnywhere, Category="Processing Options", meta=(EditCondition="bProcessController || bProcessRawInput"))
+	UPROPERTY(Config, EditAnywhere, Category="Processing Options|Require Configuration", meta=(EditCondition="bProcessController || bProcessRawInput", ConfigRestartRequired = true))
 	bool bSpecialDevicesRequireExplicitDeviceConfiguration = true;
 
 	/** 
@@ -315,7 +315,7 @@ public:
 	* 
 	* Default: True
 	*/
-	UPROPERTY(Config, EditAnywhere, Category="Processing Options")
+	UPROPERTY(Config, EditAnywhere, Category="Processing Options", meta = (ConfigRestartRequired = true))
 	bool bProcessGamepad = true;
 
 	/** 
@@ -326,7 +326,7 @@ public:
 	* 
 	* Default: True
 	*/
-	UPROPERTY(Config, EditAnywhere, Category="Processing Options")
+	UPROPERTY(Config, EditAnywhere, Category="Processing Options", meta = (ConfigRestartRequired = true))
 	bool bProcessKeyboard = true;
 
 	/** 
@@ -337,7 +337,7 @@ public:
 	* 
 	* Default: True
 	*/
-	UPROPERTY(Config, EditAnywhere, Category="Processing Options")
+	UPROPERTY(Config, EditAnywhere, Category="Processing Options", meta = (ConfigRestartRequired = true))
 	bool bProcessMouse = true;
 
 	/**
@@ -351,7 +351,7 @@ public:
 	 * 
 	 * Default: False
 	 */
-	UPROPERTY(EditAnywhere, Config, Category = "Processing Options", meta=(DisplayName="Process Racing Wheel (Experimental)"))
+	UPROPERTY(EditAnywhere, Config, Category = "Processing Options|Require Configuration", meta=(ConfigRestartRequired = true, DisplayName="Process Racing Wheel (Experimental)"))
 	bool bProcessRacingWheel = false;
 
 	/** The default racing wheel deadzone */
@@ -367,7 +367,7 @@ public:
 };
 
 /**
-* Settings related to theGame Input device interface. 
+* Settings related to the Game Input device interface. 
 * 
 * These will allow you to enable and disable specific types of input devices within Game Input
 * as well as configure key mappings for generic controller types based on their unique vendor/product ID's.
