Do not reset the LocalApp Device ID when you set the Game Input device to null. Resetting it would mean that you would not be able to find that device again upon re-connection IGameInputDeviceInterface::GetDeviceData because the APP_LOCAL_DEVICE_ID would be reset to nothing, even though the entry still exists in the array

#rnx

diff --git a/Engine/Plugins/Runtime/GameInput/Source/GameInputBase/Private/GameInputDeviceContainer.cpp b/Engine/Plugins/Runtime/GameInput/Source/GameInputBase/Private/GameInputDeviceContainer.cpp
--- a/Engine/Plugins/Runtime/GameInput/Source/GameInputBase/Private/GameInputDeviceContainer.cpp
+++ b/Engine/Plugins/Runtime/GameInput/Source/GameInputBase/Private/GameInputDeviceContainer.cpp
@@ -34,7 +34,17 @@ FGameInputDeviceContainer::FGameInputDeviceContainer(
 	, AssignedDeviceId(InDeviceId)
 	, IgnoreReadingTimestamp(0)
 {
-	SetGameInputDevice(InDevice);
+	if (!InDevice)
+	{
+		ensureAlwaysMsgf(false, TEXT("A Game Input container was created without a valid IGameInputDevice! This container will fail to process any input and we should not have gotten here."));
+		return;
+	}
+
+	// Initalize the App Local ID, which should never change on this container
+	if (const GameInputDeviceInfo* Info = InDevice->GetDeviceInfo())
+	{
+		LocalDeviceId = Info->deviceId;
+	}
 }
 
 void FGameInputDeviceContainer::SetMessageHandler(const TSharedRef<FGenericApplicationMessageHandler>& InMessageHandler)
@@ -164,6 +174,18 @@ const GameInputKind FGameInputDeviceContainer::ProcessInput(IGameInput* GameInpu
 {
 	GameInputKind OutProcessedInputKinds = GameInputKindUnknown;
 
+	// If we don't have a valid IGameInputDevice, then this device has been disconnected and there
+	// is no need to attempt to get any Game Input readings from it. 
+	//	
+	// Calling GetCurrentReading with a null IGameInputDevice will actually return _all_ game input readings, which 
+	// we don't want to process. We only care about readings associated with this container's device.
+	// 
+	// @see https://learn.microsoft.com/en-us/gaming/gdk/_content/gc/reference/input/gameinput/interfaces/igameinput/methods/igameinput_getcurrentreading
+	if (!Device)
+	{
+		return OutProcessedInputKinds;
+	}
+
 	const bool bDoNotProcessDuplicateCapabilitiesForSingleUser = GetDefault<UGameInputDeveloperSettings>()->bDoNotProcessDuplicateCapabilitiesForSingleUser;
 
 	// keep reading the input snapshots for this device
@@ -276,7 +298,7 @@ const GameInputKind FGameInputDeviceContainer::ProcessInput(IGameInput* GameInpu
 		}
 
 		// if this was the last reading of the frame, then Reading is going to be null
-		// we dont need to track anything here because it happened on the previous iteration
+		// we don't need to track anything here because it happened on the previous iteration
 		if (bIsLastReadingOfFrame)
 		{
 			break;
@@ -319,24 +341,28 @@ IGameInputDevice* FGameInputDeviceContainer::GetGameInputDevice() const
 
 void FGameInputDeviceContainer::SetGameInputDevice(IGameInputDevice* InDevice)
 {
+	// If the devices are the same, we can early exit. Nothing needs to happen
+	if (InDevice == Device)
+	{
+		return;
+	}
+
 	Device = InDevice;
+
 	if (Device)
 	{
+		// Ensure that the local app device ID is the same as it was before
+		// Every input device that is connected to game input has a unique App Local ID, so if we are 
+		// using the same container then it should be the same.
+		// 
+		// In this case, we would get here if the IGameInputDevice pointer was set to null upon disconnection, 
+		// and then the device was re-connected.
 		if (const GameInputDeviceInfo* Info = InDevice->GetDeviceInfo())
 		{
-			LocalDeviceId = Info->deviceId;
-		}
-		else
-		{
-			// We shouldn't be able to get here, but handle it just in case.
-			ensure(false);
-			LocalDeviceId = {};
+			const bool bAppIdsAreTheSame = (FMemory::Memcmp(&Info->deviceId, &LocalDeviceId, sizeof(LocalDeviceId)) == 0);
+			ensure(bAppIdsAreTheSame);
 		}
 	}
-	else
-	{
-		LocalDeviceId = {};
-	}
 }
 
 APP_LOCAL_DEVICE_ID FGameInputDeviceContainer::GetGameInputDeviceId() const
@@ -364,4 +390,14 @@ FInputDeviceId FGameInputDeviceContainer::GetDeviceId() const
 	return AssignedDeviceId;
 }
 
-#endif	// GAME_INPUT_SUPPORT
+uint64 FGameInputDeviceContainer::GetLastReadingTimestamp() const 
+{ 
+	return LastReadingTimestamp; 
+}
+
+const int32 FGameInputDeviceContainer::GetNumberOfProcessors() const
+{
+	return Processors.Num();
+}
+
+#endif	// GAME_INPUT_SUPPORT
\ No newline at end of file
diff --git a/Engine/Plugins/Runtime/GameInput/Source/GameInputBase/Private/IGameInputDeviceInterface.cpp b/Engine/Plugins/Runtime/GameInput/Source/GameInputBase/Private/IGameInputDeviceInterface.cpp
--- a/Engine/Plugins/Runtime/GameInput/Source/GameInputBase/Private/IGameInputDeviceInterface.cpp
+++ b/Engine/Plugins/Runtime/GameInput/Source/GameInputBase/Private/IGameInputDeviceInterface.cpp
@@ -443,6 +443,11 @@ FGameInputDeviceContainer* IGameInputDeviceInterface::GetOrCreateDeviceData(IGam
 	// Check if we already have some device data about this...
 	if (FGameInputDeviceContainer* ExistingDevice = GetDeviceData(InDevice))
 	{
+		// For existing devices, we want to ensure that their IGameInputDevice pointer matches up with what was given.
+		// This may be the case if you disconnect and then reconnect a device, because we can still find it's associated
+		// FGameInputDeviceContainer based on the  APP_LOCAL_DEVICE_ID, but the IGameInputDevice pointer would be null.
+		ExistingDevice->SetGameInputDevice(InDevice);
+
 		return ExistingDevice;
 	}
 
diff --git a/Engine/Plugins/Runtime/GameInput/Source/GameInputBase/Public/GameInputDeviceContainer.h b/Engine/Plugins/Runtime/GameInput/Source/GameInputBase/Public/GameInputDeviceContainer.h
--- a/Engine/Plugins/Runtime/GameInput/Source/GameInputBase/Public/GameInputDeviceContainer.h
+++ b/Engine/Plugins/Runtime/GameInput/Source/GameInputBase/Public/GameInputDeviceContainer.h
@@ -53,10 +53,23 @@ public:
 	/** Reset any input state that is necessary. This would be called when the application is no longer constrained for example */
 	void ClearInputState(IGameInput* GameInput);
 
+	/**
+	* Returns a pointer to the IGameInputDevice that this device is associated with
+	* 
+	* If this is null, then this is a container for an input device which has been disconnected.
+	*/
 	IGameInputDevice* GetGameInputDevice() const;
+
+	/**
+	* Update the IGameInputDevice pointer that this container should use. 
+	*
+	* This IGameInputDevice should have the same Local App ID as this container already, otherwise there will be an ensure
+	* 
+	* @param InDevice	The new IGameInputDevice that this container should use. 
+	*/
 	void SetGameInputDevice(IGameInputDevice* InDevice);
 
-	/** Returns the unique local device ID of this IGameInputDevice from the GameInput SDK. */
+	/** Returns the unique local device ID of this IGameInputDevice from the GameInput API. */
 	APP_LOCAL_DEVICE_ID GetGameInputDeviceId() const;
 
 	void SetPlatformUserId(const FPlatformUserId InUserId);
@@ -65,14 +78,14 @@ public:
 	void SetInputDeviceId(const FInputDeviceId InDeviceId);
 	FInputDeviceId GetDeviceId() const;
 
-	uint64 GetLastReadingTimestamp() const { return LastReadingTimestamp; }
+	uint64 GetLastReadingTimestamp() const;
 
 	/** 
 	* Returns the number of processors that this device container currently has.
 	*
 	* Note: If this is zero, then this device can't possibly fire any input events.
 	*/
-	const int32 GetNumberOfProcessors() const { return Processors.Num(); }
+	const int32 GetNumberOfProcessors() const;
 
 protected:
 
@@ -94,7 +107,13 @@ protected:
 	/** The Input Device ID that this Game Input Device is associated with */
 	FInputDeviceId AssignedDeviceId;
 
-	/** The unique ID that is associated with this device from the GameInput API */
+	/** 
+	* The unique ID that is associated with this device from the GameInput API.
+	* This will be set only once upon construction of this container.
+	* 
+	* @see IGameInputDeviceInterface::GetDeviceData
+	* @see https://learn.microsoft.com/en-us/gaming/gdk/_content/gc/reference/system/xuser/structs/app_local_device_id
+	*/
 	APP_LOCAL_DEVICE_ID LocalDeviceId;
 
 	/** Input Device processors to handle input for any specific types required (gamepad, racing wheel, KBM, etc) */
