All you to set certain raw input analog values to not set the input device scope

diff --git a/Engine/Plugins/Runtime/GameInput/Source/GameInputBase/Private/GameInputDeviceProcessor.cpp b/Engine/Plugins/Runtime/GameInput/Source/GameInputBase/Private/GameInputDeviceProcessor.cpp
--- a/Engine/Plugins/Runtime/GameInput/Source/GameInputBase/Private/GameInputDeviceProcessor.cpp
+++ b/Engine/Plugins/Runtime/GameInput/Source/GameInputBase/Private/GameInputDeviceProcessor.cpp
@@ -106,7 +106,7 @@ const FString& IGameInputDeviceProcessor::GetHardwareDeviceIdentifierName(const
 	return ID_XboxOne;
 }
 
-void IGameInputDeviceProcessor::OnControllerAnalog(const FGameInputEventParams& Params, const FName& GamePadKey, float NewAxisValueNormalized, float OldAxisValueNormalized, float DeadZone)
+void IGameInputDeviceProcessor::OnControllerAnalog(const FGameInputEventParams& Params, const FName& GamePadKey, float NewAxisValueNormalized, float OldAxisValueNormalized, float DeadZone, const bool bSetDeviceScope /*= true*/)
 {
 	if (OldAxisValueNormalized != NewAxisValueNormalized || FMath::Abs(NewAxisValueNormalized) > DeadZone)
 	{		
@@ -121,8 +121,15 @@ void IGameInputDeviceProcessor::OnControllerAnalog(const FGameInputEventParams&
 		// to create a new slate user based on the index if it doesn't already exist
 		if (Params.PlatformUserId.IsValid() && Params.InputDeviceId.IsValid())
 		{
-			FInputDeviceScope InputScope(nullptr, UE::GameInput::InputClassName, IPlatformInputDeviceMapper::Get().GetUserIndexForPlatformUser(Params.PlatformUserId), GetHardwareDeviceIdentifierName(Params));
-			Params.MessageHandler->OnControllerAnalog(GamePadKey, Params.PlatformUserId, Params.InputDeviceId, NewAxisValueNormalized);
+			if (bSetDeviceScope)
+			{
+				FInputDeviceScope InputScope(nullptr, UE::GameInput::InputClassName, IPlatformInputDeviceMapper::Get().GetUserIndexForPlatformUser(Params.PlatformUserId), GetHardwareDeviceIdentifierName(Params));
+				Params.MessageHandler->OnControllerAnalog(GamePadKey, Params.PlatformUserId, Params.InputDeviceId, NewAxisValueNormalized);
+			}
+			else
+			{
+				Params.MessageHandler->OnControllerAnalog(GamePadKey, Params.PlatformUserId, Params.InputDeviceId, NewAxisValueNormalized);
+			}
 		}		
 	}
 }
@@ -1694,7 +1701,13 @@ bool FGameInputRawDeviceProcessor::ProcessRawInputValueAsAanalog(const FGameInpu
 		return false;
 	}
 
-	OnControllerAnalog(Params, AxisData->KeyName, CurrentValueFloat, PreviousValueFloat, UE::GameInput::GamepadLeftStickDeadzone);
+	OnControllerAnalog(
+		Params, 
+		AxisData->KeyName, 
+		CurrentValueFloat, 
+		PreviousValueFloat, 
+		UE::GameInput::GamepadLeftStickDeadzone, 
+		/* bShouldSetDeviceScope = */!AxisData->bIgnoreAnalogInputDeviceScopeForThisRawReport);
 
 	// We had a reading as long as it is non-zero
 	return CurrentValueFloat != 0.0f;
diff --git a/Engine/Plugins/Runtime/GameInput/Source/GameInputBase/Public/GameInputDeveloperSettings.h b/Engine/Plugins/Runtime/GameInput/Source/GameInputBase/Public/GameInputDeveloperSettings.h
--- a/Engine/Plugins/Runtime/GameInput/Source/GameInputBase/Public/GameInputDeveloperSettings.h
+++ b/Engine/Plugins/Runtime/GameInput/Source/GameInputBase/Public/GameInputDeveloperSettings.h
@@ -131,6 +131,18 @@ struct GAMEINPUTBASE_API FGameInputRawDeviceReportData
 	UPROPERTY(EditAnywhere, Config, Category = "Game Input Device", meta=(EditCondition="TranslationBehavior != ERawDeviceReportTranslationBehavior::TreatAsButtonBitmask", EditConditionHides))
 	FName KeyName = NAME_None;
 
+	/**
+	* If true, then this raw report data will NOT set the FInputDeviceScope.
+	* 
+	* This may be desirable if you have a very sensitive axis value or something which may be constantly reporting
+	* that you don't necessarily want to be used to determine things like Common UI icons or other systems who read
+	* device scope.
+	* 
+	* Default: False.
+	*/
+	UPROPERTY(EditAnywhere, Config, Category = "Game Input Device", meta=(EditCondition="TranslationBehavior == ERawDeviceReportTranslationBehavior::TreatAsAnalog || TranslationBehavior == ERawDeviceReportTranslationBehavior::TreatAsTrigger", EditConditionHides))
+	bool bIgnoreAnalogInputDeviceScopeForThisRawReport = false;
+
 	/**
 	 * Options for how we should interpret the raw uint8 value from RawInput when telling the engine about it
 	 */
diff --git a/Engine/Plugins/Runtime/GameInput/Source/GameInputBase/Public/GameInputDeviceProcessor.h b/Engine/Plugins/Runtime/GameInput/Source/GameInputBase/Public/GameInputDeviceProcessor.h
--- a/Engine/Plugins/Runtime/GameInput/Source/GameInputBase/Public/GameInputDeviceProcessor.h
+++ b/Engine/Plugins/Runtime/GameInput/Source/GameInputBase/Public/GameInputDeviceProcessor.h
@@ -116,7 +116,7 @@ protected:
 	const FString& GetHardwareDeviceIdentifierName(const IGameInputDeviceProcessor::FGameInputEventParams& Params) const;
 
 	/** A general use function to call the message handler and tell it about a controller analog key being used */
-	void OnControllerAnalog(const FGameInputEventParams& Params, const FName& GamePadKey, float NewAxisValueNormalized, float OldAxisValueNormalized, float DeadZone);
+	void OnControllerAnalog(const FGameInputEventParams& Params, const FName& GamePadKey, float NewAxisValueNormalized, float OldAxisValueNormalized, float DeadZone, const bool bSetDeviceScope = true);
 
 	/**
 	 * Helper function for processing the button states of Game input.
