Add a comment about using Gamepad processing on Game Input for windows to make sure users know to disable XInput or other input device plugins to avoid duplicate input events.

#rnx

diff --git a/Engine/Plugins/Runtime/GameInput/Source/GameInputBase/Public/GameInputDeveloperSettings.h b/Engine/Plugins/Runtime/GameInput/Source/GameInputBase/Public/GameInputDeveloperSettings.h
--- a/Engine/Plugins/Runtime/GameInput/Source/GameInputBase/Public/GameInputDeveloperSettings.h
+++ b/Engine/Plugins/Runtime/GameInput/Source/GameInputBase/Public/GameInputDeveloperSettings.h
@@ -325,6 +325,9 @@ public:
 	/** 
 	* If true, then we will process the GameInputKindGamepad type.
 	* 
+	* Note: If you are using Game Input on Windows where there are other Input Device module plugins (XInput, WinDualShock, etc) 
+	* you should disable those to use this. Otherwise, there will be "duplicate" gamepad input events.
+	* 
 	* Default: True
 	*/
 	UPROPERTY(Config, EditAnywhere, Category="Processing Options", meta = (ConfigRestartRequired = true))
@@ -334,7 +337,7 @@ public:
 	* If true, then we will process the GameInputKindKeyboard type.
 	* 
 	* Note: You likely do not want this on for Windows targets, as those input events are already processed via the 
-	* WindowsApplication. This should really only be used for console targets.
+	* WindowsApplication. This should really only be used for console targets. This is currently disabled on Windows.
 	* 
 	* Default: True
 	*/
@@ -345,7 +348,7 @@ public:
 	* If true, then we will process the GameInputKindMouse type.
 	* 
 	* Note: You likely do not want this on for Windows targets, as those input events are already processed via the 
-	* WindowsApplication. This should really only be used for console targets.
+	* WindowsApplication. This should really only be used for console targets. This is currently disabled on Windows.
 	* 
 	* Default: True
 	*/
