Arcade Stick support in Game Input

diff --git a/Engine/Plugins/Runtime/GameInput/Source/GameInputBase/Private/GameInputBaseModule.cpp b/Engine/Plugins/Runtime/GameInput/Source/GameInputBase/Private/GameInputBaseModule.cpp
--- a/Engine/Plugins/Runtime/GameInput/Source/GameInputBase/Private/GameInputBaseModule.cpp
+++ b/Engine/Plugins/Runtime/GameInput/Source/GameInputBase/Private/GameInputBaseModule.cpp
@@ -197,13 +197,6 @@ void FGameInputBaseModule::InitializeGameInputKeys()
 	//
 
 	// Button Types
-	EKeys::AddKey(FKeyDetails(FGameInputKeys::ArcadeStick_None, LOCTEXT("GameInput_ArcadeStick_None", "Game Input Arcade Stick None"), FKeyDetails::GamepadKey, MenuCategory));
-	EKeys::AddKey(FKeyDetails(FGameInputKeys::ArcadeStick_Menu, LOCTEXT("GameInput_ArcadeStick_Menu", "Game Input Arcade Stick Menu"), FKeyDetails::GamepadKey, MenuCategory));
-	EKeys::AddKey(FKeyDetails(FGameInputKeys::ArcadeStick_View, LOCTEXT("GameInput_ArcadeStick_View", "Game Input Arcade Stick View"), FKeyDetails::GamepadKey, MenuCategory));
-	EKeys::AddKey(FKeyDetails(FGameInputKeys::ArcadeStick_Up, LOCTEXT("GameInput_ArcadeStick_Up", "Game Input Arcade Stick Up"), FKeyDetails::GamepadKey, MenuCategory));
-	EKeys::AddKey(FKeyDetails(FGameInputKeys::ArcadeStick_Down, LOCTEXT("GameInput_ArcadeStick_Down", "Game Input Arcade Stick Down"), FKeyDetails::GamepadKey, MenuCategory));
-	EKeys::AddKey(FKeyDetails(FGameInputKeys::ArcadeStick_Left, LOCTEXT("GameInput_ArcadeStick_Left", "Game Input Arcade Stick Left"), FKeyDetails::GamepadKey, MenuCategory));
-	EKeys::AddKey(FKeyDetails(FGameInputKeys::ArcadeStick_Right, LOCTEXT("GameInput_ArcadeStick_Right", "Game Input Arcade Stick Right"), FKeyDetails::GamepadKey, MenuCategory));
 	EKeys::AddKey(FKeyDetails(FGameInputKeys::ArcadeStick_Action1, LOCTEXT("GameInput_ArcadeStick_Action1", "Game Input Arcade Stick Action 1"), FKeyDetails::GamepadKey, MenuCategory));
 	EKeys::AddKey(FKeyDetails(FGameInputKeys::ArcadeStick_Action2, LOCTEXT("GameInput_ArcadeStick_Action2", "Game Input Arcade Stick Action 2"), FKeyDetails::GamepadKey, MenuCategory));
 	EKeys::AddKey(FKeyDetails(FGameInputKeys::ArcadeStick_Action3, LOCTEXT("GameInput_ArcadeStick_Action3", "Game Input Arcade Stick Action 3"), FKeyDetails::GamepadKey, MenuCategory));
diff --git a/Engine/Plugins/Runtime/GameInput/Source/GameInputBase/Private/GameInputDeviceContainer.cpp b/Engine/Plugins/Runtime/GameInput/Source/GameInputBase/Private/GameInputDeviceContainer.cpp
--- a/Engine/Plugins/Runtime/GameInput/Source/GameInputBase/Private/GameInputDeviceContainer.cpp
+++ b/Engine/Plugins/Runtime/GameInput/Source/GameInputBase/Private/GameInputDeviceContainer.cpp
@@ -157,17 +157,18 @@ void FGameInputDeviceContainer::InitalizeDeviceProcessors_Impl()
 		Processors.Emplace(MakeShared<FGameInputRacingWheelProcessor>());
 	}
 
+	// Arcade sticks
+	if (DeviceKind & AllowedGameInputKinds & GameInputKindArcadeStick)
+	{
+		Processors.Emplace(MakeShared<FGameInputArcadeStickProcessor>());
+	}
+
 	// TODO: Implement these kinds of input! 
 
 	//if (DeviceKind & AllowedGameInputKinds & GameInputKindFlightStick)
 	//{
 
 	//}
-
-	//if (DeviceKind & AllowedGameInputKinds & GameInputKindArcadeStick)
-	//{
-
-	//}
 }
 
 const GameInputKind FGameInputDeviceContainer::ProcessInput(IGameInput* GameInput, const GameInputKind CurrentSupportedKind, const GameInputKind ProcessedKindsForPlatformUserThisFrame)
diff --git a/Engine/Plugins/Runtime/GameInput/Source/GameInputBase/Private/GameInputDeviceProcessor.cpp b/Engine/Plugins/Runtime/GameInput/Source/GameInputBase/Private/GameInputDeviceProcessor.cpp
--- a/Engine/Plugins/Runtime/GameInput/Source/GameInputBase/Private/GameInputDeviceProcessor.cpp
+++ b/Engine/Plugins/Runtime/GameInput/Source/GameInputBase/Private/GameInputDeviceProcessor.cpp
@@ -1971,4 +1971,108 @@ GameInputKind FGameInputRacingWheelProcessor::GetSupportedReadingKind() const
 	return GameInputKindRacingWheel;
 }
 
+//////////////////////////////////////////////////
+// Arcade Sticks
+
+namespace UE::GameInput
+{
+	/** A map of uint32 GameInput button bitmask flags to the associated Unreal Engine FKey name. */
+	static const TMap<uint32, FGamepadKeyNames::Type>& GetArcadeStickButtonMap()
+	{
+		static const TMap<uint32, FGamepadKeyNames::Type> GamepadButtonMap
+		{
+			// Generic gamepad buttons
+			{ static_cast<uint32>(GameInputArcadeStickButtons::GameInputArcadeStickNone), FGamepadKeyNames::Invalid},
+			{ static_cast<uint32>(GameInputArcadeStickButtons::GameInputArcadeStickMenu), FGamepadKeyNames::SpecialRight },
+			{ static_cast<uint32>(GameInputArcadeStickButtons::GameInputArcadeStickView), FGamepadKeyNames::SpecialLeft },
+			{ static_cast<uint32>(GameInputArcadeStickButtons::GameInputArcadeStickUp), FGamepadKeyNames::DPadUp },
+			{ static_cast<uint32>(GameInputArcadeStickButtons::GameInputArcadeStickDown), FGamepadKeyNames::DPadDown },
+			{ static_cast<uint32>(GameInputArcadeStickButtons::GameInputArcadeStickLeft), FGamepadKeyNames::DPadLeft },
+			{ static_cast<uint32>(GameInputArcadeStickButtons::GameInputArcadeStickRight), FGamepadKeyNames::DPadRight },
+
+			// Unique to arcade sticks
+			{ static_cast<uint32>(GameInputArcadeStickButtons::GameInputArcadeStickAction1), FGameInputKeys::ArcadeStick_Action1.GetFName() },
+			{ static_cast<uint32>(GameInputArcadeStickButtons::GameInputArcadeStickAction2), FGameInputKeys::ArcadeStick_Action2.GetFName() },
+			{ static_cast<uint32>(GameInputArcadeStickButtons::GameInputArcadeStickAction3), FGameInputKeys::ArcadeStick_Action3.GetFName() },
+			{ static_cast<uint32>(GameInputArcadeStickButtons::GameInputArcadeStickAction4), FGameInputKeys::ArcadeStick_Action4.GetFName() },
+			{ static_cast<uint32>(GameInputArcadeStickButtons::GameInputArcadeStickAction5), FGameInputKeys::ArcadeStick_Action5.GetFName() },
+			{ static_cast<uint32>(GameInputArcadeStickButtons::GameInputArcadeStickAction6), FGameInputKeys::ArcadeStick_Action6.GetFName() },
+			{ static_cast<uint32>(GameInputArcadeStickButtons::GameInputArcadeStickSpecial1), FGameInputKeys::ArcadeStick_Special1.GetFName() },
+			{ static_cast<uint32>(GameInputArcadeStickButtons::GameInputArcadeStickSpecial2), FGameInputKeys::ArcadeStick_Special2.GetFName() }
+		};
+		return GamepadButtonMap;
+	}
+};
+
+
+FGameInputArcadeStickProcessor::FGameInputArcadeStickProcessor()
+	: IGameInputDeviceProcessor()
+{
+	FMemory::Memset(PreviousState, 0);
+	FMemory::Memset(RepeatTime, 0);
+}
+
+bool FGameInputArcadeStickProcessor::ProcessInput(const FGameInputEventParams& Params)
+{
+	bool bRes = false;
+
+	// Can't do anything for an invalid platform user or no reading.
+	if (!Params.PlatformUserId.IsValid() || !Params.Reading)
+	{
+		return bRes;
+	}
+
+	GameInputArcadeStickState StickState;
+	if (!Params.Reading->GetArcadeStickState(&StickState))
+	{
+		return bRes;
+	}
+
+	const uint32 CurrentButtonHeldMask = static_cast<uint32>(StickState.buttons);
+	uint32 LastButtonHeldMask = static_cast<uint32>(PreviousState.buttons);
+
+	EvaluateButtonStates(
+		Params,
+		CurrentButtonHeldMask,
+		OUT LastButtonHeldMask,
+		RepeatTime,
+		UE::GameInput::GetArcadeStickButtonMap(),
+		MaxSupportedButtons);
+
+	// Keep track of the button state so that we can compare it next time it is processed
+	PreviousState = StickState;
+
+	return true;
+}
+
+void FGameInputArcadeStickProcessor::ClearState(const FGameInputEventParams& Params)
+{	
+	// Can't do anything for an invalid platform user
+	if (!Params.PlatformUserId.IsValid())
+	{
+		return;
+	}
+
+	// Evaluate the buttons as if none have been pressed (i.e. the button mask is 0)
+	const uint32 CurrentButtonHeldMask = 0x00;
+	uint32 LastButtonHeldMask = static_cast<uint32>(PreviousState.buttons);
+
+	EvaluateButtonStates(
+		Params,
+		CurrentButtonHeldMask,
+		OUT LastButtonHeldMask,
+		RepeatTime,
+		UE::GameInput::GetArcadeStickButtonMap(),
+		MaxSupportedButtons);
+
+	// Zero out the repeat and previous states
+	FMemory::Memset(PreviousState, 0);
+	FMemory::Memset(RepeatTime, 0);
+}
+
+GameInputKind FGameInputArcadeStickProcessor::GetSupportedReadingKind() const
+{
+	return GameInputKindArcadeStick;
+}
+
 #endif	// GAME_INPUT_SUPPORT
\ No newline at end of file
diff --git a/Engine/Plugins/Runtime/GameInput/Source/GameInputBase/Private/GameInputKeyTypes.cpp b/Engine/Plugins/Runtime/GameInput/Source/GameInputBase/Private/GameInputKeyTypes.cpp
--- a/Engine/Plugins/Runtime/GameInput/Source/GameInputBase/Private/GameInputKeyTypes.cpp
+++ b/Engine/Plugins/Runtime/GameInput/Source/GameInputBase/Private/GameInputKeyTypes.cpp
@@ -43,13 +43,6 @@ const FKey FGameInputKeys::FlightStick_FireSecondary("GameInput_FlightStick_Fire
 //
 
 // Button Types
-const FKey FGameInputKeys::ArcadeStick_None("GameInput_ArcadeStick_None");
-const FKey FGameInputKeys::ArcadeStick_Menu("GameInput_ArcadeStick_Menu");
-const FKey FGameInputKeys::ArcadeStick_View("GameInput_ArcadeStick_View");
-const FKey FGameInputKeys::ArcadeStick_Up("GameInput_ArcadeStick_Up");
-const FKey FGameInputKeys::ArcadeStick_Down("GameInput_ArcadeStick_Down");
-const FKey FGameInputKeys::ArcadeStick_Left("GameInput_ArcadeStick_Left");
-const FKey FGameInputKeys::ArcadeStick_Right("GameInput_ArcadeStick_Right");
 const FKey FGameInputKeys::ArcadeStick_Action1("GameInput_ArcadeStick_Action1");
 const FKey FGameInputKeys::ArcadeStick_Action2("GameInput_ArcadeStick_Action2");
 const FKey FGameInputKeys::ArcadeStick_Action3("GameInput_ArcadeStick_Action3");
diff --git a/Engine/Plugins/Runtime/GameInput/Source/GameInputBase/Private/IGameInputDeviceInterface.cpp b/Engine/Plugins/Runtime/GameInput/Source/GameInputBase/Private/IGameInputDeviceInterface.cpp
--- a/Engine/Plugins/Runtime/GameInput/Source/GameInputBase/Private/IGameInputDeviceInterface.cpp
+++ b/Engine/Plugins/Runtime/GameInput/Source/GameInputBase/Private/IGameInputDeviceInterface.cpp
@@ -286,6 +286,11 @@ GameInputKind IGameInputDeviceInterface::GetCurrentGameInputKindSupport() const
 		RegisterInputKindMask |= GameInputKindRacingWheel;
 	}
 
+	if (PlatformSettings->bProcessArcadeStick)
+	{
+		RegisterInputKindMask |= GameInputKindArcadeStick;
+	}
+
 	// TODO: Future expansion of GameInput devices!
 	/*if (Settings->bProcessArcadeStick)
 	{
diff --git a/Engine/Plugins/Runtime/GameInput/Source/GameInputBase/Public/GameInputDeveloperSettings.h b/Engine/Plugins/Runtime/GameInput/Source/GameInputBase/Public/GameInputDeveloperSettings.h
--- a/Engine/Plugins/Runtime/GameInput/Source/GameInputBase/Public/GameInputDeveloperSettings.h
+++ b/Engine/Plugins/Runtime/GameInput/Source/GameInputBase/Public/GameInputDeveloperSettings.h
@@ -379,6 +379,16 @@ public:
 	*/
 	UPROPERTY(EditAnywhere, Config, Category = "Device Settings|Racing Wheel", meta=(EditCondition="bProcessRacingWheel"))
 	float RacingWheelDeadzone = DefaultRacingWheelDeadzone;
+
+	/**
+	* If true, then we will process the GameInputKindArcadeStick type.
+	* 
+	* Note: This is experimental!
+	* 
+	* Default: False
+	*/
+	UPROPERTY(EditAnywhere, Config, Category = "Processing Options", meta=(ConfigRestartRequired = true, DisplayName="Process Arcade Sticks (Experimental)"))
+	bool bProcessArcadeStick = false;
 };
 
 /**
diff --git a/Engine/Plugins/Runtime/GameInput/Source/GameInputBase/Public/GameInputDeviceProcessor.h b/Engine/Plugins/Runtime/GameInput/Source/GameInputBase/Public/GameInputDeviceProcessor.h
--- a/Engine/Plugins/Runtime/GameInput/Source/GameInputBase/Public/GameInputDeviceProcessor.h
+++ b/Engine/Plugins/Runtime/GameInput/Source/GameInputBase/Public/GameInputDeviceProcessor.h
@@ -495,4 +495,32 @@ protected:
 	double RepeatTime[MaxSupportedButtons];
 };
 
+
+/**
+* Processor for the GameInputKindArcadeStick type.
+* 
+* These are typically accessories to "fighting" style games, or other arcade style games.
+*/
+class GAMEINPUTBASE_API FGameInputArcadeStickProcessor : public IGameInputDeviceProcessor
+{
+public:
+	FGameInputArcadeStickProcessor();
+
+protected:
+
+	virtual bool ProcessInput(const FGameInputEventParams& Params) override;
+	virtual void ClearState(const FGameInputEventParams& Params) override;
+	virtual GameInputKind GetSupportedReadingKind() const override;
+
+	/** The previously processed arcade stick state. */
+	GameInputArcadeStickState PreviousState;
+
+	/**
+	* Array of repeat times to calculate if a button has been held long enough
+	* to receive an IE_REPEAT event.
+	*/
+	static const uint32 MaxSupportedButtons = 16;
+	double RepeatTime[MaxSupportedButtons];
+};
+
 #endif	// GAME_INPUT_SUPPORT
\ No newline at end of file
diff --git a/Engine/Plugins/Runtime/GameInput/Source/GameInputBase/Public/GameInputKeyTypes.h b/Engine/Plugins/Runtime/GameInput/Source/GameInputBase/Public/GameInputKeyTypes.h
--- a/Engine/Plugins/Runtime/GameInput/Source/GameInputBase/Public/GameInputKeyTypes.h
+++ b/Engine/Plugins/Runtime/GameInput/Source/GameInputBase/Public/GameInputKeyTypes.h
@@ -54,13 +54,6 @@ struct GAMEINPUTBASE_API FGameInputKeys
 	//
 
 	// Button Types
-	static const FKey ArcadeStick_None;
-	static const FKey ArcadeStick_Menu;
-	static const FKey ArcadeStick_View;
-	static const FKey ArcadeStick_Up;
-	static const FKey ArcadeStick_Down;
-	static const FKey ArcadeStick_Left;
-	static const FKey ArcadeStick_Right;
 	static const FKey ArcadeStick_Action1;
 	static const FKey ArcadeStick_Action2;
 	static const FKey ArcadeStick_Action3;
