Add flight stick processing to GameInput

diff --git a/Engine/Plugins/Runtime/GameInput/Source/GameInputBase/Private/GameInputDeviceContainer.cpp b/Engine/Plugins/Runtime/GameInput/Source/GameInputBase/Private/GameInputDeviceContainer.cpp
--- a/Engine/Plugins/Runtime/GameInput/Source/GameInputBase/Private/GameInputDeviceContainer.cpp
+++ b/Engine/Plugins/Runtime/GameInput/Source/GameInputBase/Private/GameInputDeviceContainer.cpp
@@ -163,12 +163,11 @@ void FGameInputDeviceContainer::InitalizeDeviceProcessors_Impl()
 		Processors.Emplace(MakeShared<FGameInputArcadeStickProcessor>());
 	}
 
-	// TODO: Implement these kinds of input! 
-
-	//if (DeviceKind & AllowedGameInputKinds & GameInputKindFlightStick)
-	//{
-
-	//}
+	// Flight sticks
+	if (DeviceKind & AllowedGameInputKinds & GameInputKindFlightStick)
+	{
+		Processors.Emplace(MakeShared<FGameInputFlightStickProcessor>());
+	}
 }
 
 const GameInputKind FGameInputDeviceContainer::ProcessInput(IGameInput* GameInput, const GameInputKind CurrentSupportedKind, const GameInputKind ProcessedKindsForPlatformUserThisFrame)
diff --git a/Engine/Plugins/Runtime/GameInput/Source/GameInputBase/Private/GameInputDeviceProcessor.cpp b/Engine/Plugins/Runtime/GameInput/Source/GameInputBase/Private/GameInputDeviceProcessor.cpp
--- a/Engine/Plugins/Runtime/GameInput/Source/GameInputBase/Private/GameInputDeviceProcessor.cpp
+++ b/Engine/Plugins/Runtime/GameInput/Source/GameInputBase/Private/GameInputDeviceProcessor.cpp
@@ -2075,4 +2075,194 @@ GameInputKind FGameInputArcadeStickProcessor::GetSupportedReadingKind() const
 	return GameInputKindArcadeStick;
 }
 
-#endif	// GAME_INPUT_SUPPORT
\ No newline at end of file
+//////////////////////////////////////////////////////////////////////
+// FGameInputFlightStickProcessor
+
+namespace UE::GameInput
+{
+	/** A map of uint32 GameInput button bitmask flags to the associated Unreal Engine FKey name. */
+	static const TMap<uint32, FGamepadKeyNames::Type>& GetFlightStickButtonMap()
+	{
+		static const TMap<uint32, FGamepadKeyNames::Type> GamepadButtonMap
+		{
+			// Generic gamepad buttons
+			{ static_cast<uint32>(GameInputFlightStickButtons::GameInputFlightStickNone), FGamepadKeyNames::Invalid},
+			{ static_cast<uint32>(GameInputFlightStickButtons::GameInputFlightStickMenu), FGamepadKeyNames::Invalid},
+			{ static_cast<uint32>(GameInputFlightStickButtons::GameInputFlightStickView), FGamepadKeyNames::Invalid},
+			{ static_cast<uint32>(GameInputFlightStickButtons::GameInputFlightStickFirePrimary), FGamepadKeyNames::Invalid},
+			{ static_cast<uint32>(GameInputFlightStickButtons::GameInputFlightStickFireSecondary), FGamepadKeyNames::Invalid},
+		};
+		return GamepadButtonMap;
+	}
+	
+	static bool HasDifferentFlightStickAnalogs(const GameInputFlightStickState& CurrentState, const GameInputFlightStickState& PreviousState)
+	{
+		// If anything differs from the previous reading, then it has different inputs
+		return 
+			CurrentState.pitch != PreviousState.pitch ||
+			CurrentState.roll != PreviousState.roll ||
+			CurrentState.throttle != PreviousState.throttle ||
+			CurrentState.yaw != PreviousState.yaw;
+	}
+
+	inline bool IsEmptyFlightStickReading(const GameInputFlightStickState& State, const UGameInputPlatformSettings& PlatformSettings)
+	{
+		return
+			State.buttons == 0 &&
+			State.hatSwitch == 0 &&	
+			FMath::Abs(State.pitch) <= PlatformSettings.FlightStickPitchDeadzone &&
+			FMath::Abs(State.roll) <= PlatformSettings.FlightStickRollDeadzone &&
+			FMath::Abs(State.throttle) <= PlatformSettings.FlightStickThrottleDeadzone &&
+			FMath::Abs(State.yaw) <= PlatformSettings.FlightStickYawDeadzone; 
+	};
+}
+
+FGameInputFlightStickProcessor::FGameInputFlightStickProcessor()
+	: IGameInputDeviceProcessor()
+{
+	FMemory::Memset(PreviousState, 0);
+	FMemory::Memset(RepeatTime, 0);
+	
+	SwitchRepeatTimes.AddDefaulted(static_cast<uint32>(GameInputSwitchUpLeft) + 1);
+}
+
+bool FGameInputFlightStickProcessor::ProcessInput(const FGameInputEventParams& Params)
+{
+	bool bRes = false;
+	
+	// Can't do anything for an invalid platform user
+	if (!Params.PlatformUserId.IsValid() || !Params.Reading)
+	{
+		return bRes;
+	}
+
+	GameInputFlightStickState FlightStickState = {};
+	if (!Params.Reading->GetFlightStickState(&FlightStickState))
+	{
+		return bRes;
+	}
+
+	// We only want to process the buttons here, as it might get called multiple times per frame.
+	bRes |= ProcessFlightStickButtons(Params, FlightStickState);
+	
+	++NumReadingsProcessedThisFrame;
+	
+	return bRes;
+}
+
+bool FGameInputFlightStickProcessor::PostProcessInput(const FGameInputEventParams& Params)
+{
+	bool bRes = false;
+	
+	// Check if we have already processed buttons this frame. If we haven't we want to do it
+	const bool bHasProcessedAnyButtonsThisFrame = NumReadingsProcessedThisFrame > 0;	
+	NumReadingsProcessedThisFrame = 0;
+	
+	// Can't do anything for an invalid platform user
+	if (!Params.PlatformUserId.IsValid() || !Params.PreviousReading)
+	{
+		return bRes;
+	}
+
+	GameInputFlightStickState FlightStickState = {};
+	if (!Params.PreviousReading->GetFlightStickState(&FlightStickState))
+	{
+		return bRes;
+	}
+	
+	if (!bHasProcessedAnyButtonsThisFrame)
+	{
+		bRes |= ProcessFlightStickButtons(Params, FlightStickState);
+	}
+	
+	bRes |= ProcessFlightStickAnalog(Params, FlightStickState);
+	
+	return bRes;
+}
+
+void FGameInputFlightStickProcessor::ClearState(const FGameInputEventParams& Params)
+{
+	// Can't do anything for an invalid platform user
+	if (!Params.PlatformUserId.IsValid())
+	{
+		return;
+	}
+	
+	// Process input as if nothing is down (zero values for everything)
+	GameInputFlightStickState ZeroState = {};
+	ProcessFlightStickAnalog(Params, ZeroState);
+	ProcessFlightStickButtons(Params, ZeroState);
+	
+	// Zero out the repeat and previous states info to zero
+	FMemory::Memset(PreviousState, 0);
+	FMemory::Memset(RepeatTime, 0);
+}
+
+GameInputKind FGameInputFlightStickProcessor::GetSupportedReadingKind() const
+{
+	return GameInputKindFlightStick;
+}
+
+bool FGameInputFlightStickProcessor::ProcessFlightStickButtons(const FGameInputEventParams& Params, GameInputFlightStickState& State)
+{
+	const UGameInputPlatformSettings* PlatformSettings = UGameInputPlatformSettings::Get();
+
+	const bool bIsEmptyReading = UE::GameInput::IsEmptyFlightStickReading(State, *PlatformSettings);
+	const bool bWasEmptyReading = UE::GameInput::IsEmptyFlightStickReading(PreviousState, *PlatformSettings);
+	const bool bHasDifferentAnalogInput = UE::GameInput::HasDifferentFlightStickAnalogs(State, PreviousState);
+	
+	if (bIsEmptyReading && bWasEmptyReading && !bHasDifferentAnalogInput && PreviousState.buttons == 0)
+	{
+		return false;
+	}
+	
+	const uint32 CurrentButtonHeldMask = static_cast<uint32>(State.buttons);
+	uint32 LastButtonHeldMask = static_cast<uint32>(PreviousState.buttons);
+
+	EvaluateButtonStates(
+		Params,
+		CurrentButtonHeldMask,
+		OUT LastButtonHeldMask,
+		RepeatTime,
+		UE::GameInput::GetFlightStickButtonMap(),
+		MaxSupportedButtons);
+
+	// Update the previous state here
+	PreviousState.buttons = State.buttons;
+
+	// Update the hat switch
+	EvaluateSwitchState(Params, State.hatSwitch, PreviousState.hatSwitch, SwitchRepeatTimes);
+	
+	return true;
+}
+
+bool FGameInputFlightStickProcessor::ProcessFlightStickAnalog(const FGameInputEventParams& Params, GameInputFlightStickState& State)
+{
+
+	
+	const UGameInputPlatformSettings* Settings = UGameInputPlatformSettings::Get();
+	check(Settings);
+	
+	// If the analog values haven't changed, don't bother sending any events for them
+	const bool bIsEmptyReading = UE::GameInput::IsEmptyFlightStickReading(State, *Settings);
+	const bool bWasEmptyReading = UE::GameInput::IsEmptyFlightStickReading(PreviousState, *Settings);
+	const bool bHasDifferentAnalogInput = UE::GameInput::HasDifferentFlightStickAnalogs(State, PreviousState);
+	if (bIsEmptyReading && bWasEmptyReading && !bHasDifferentAnalogInput)
+	{
+		return false;
+	}
+	
+	OnControllerAnalog(Params, FGameInputKeys::FlightStick_Pitch.GetFName(), State.pitch, PreviousState.pitch, Settings->FlightStickPitchDeadzone);
+	OnControllerAnalog(Params, FGameInputKeys::FlightStick_Roll.GetFName(), State.roll, PreviousState.roll, Settings->FlightStickRollDeadzone);
+	OnControllerAnalog(Params, FGameInputKeys::FlightStick_Throttle.GetFName(), State.throttle, PreviousState.throttle, Settings->FlightStickThrottleDeadzone);
+	OnControllerAnalog(Params, FGameInputKeys::FlightStick_Yaw.GetFName(), State.yaw, PreviousState.yaw, Settings->FlightStickYawDeadzone);
+
+	PreviousState.pitch = State.pitch;
+	PreviousState.roll = State.roll;
+	PreviousState.throttle = State.throttle;
+	PreviousState.yaw = State.yaw;
+	
+	return true;
+}
+
+#endif	// GAME_INPUT_SUPPORT
diff --git a/Engine/Plugins/Runtime/GameInput/Source/GameInputBase/Private/IGameInputDeviceInterface.cpp b/Engine/Plugins/Runtime/GameInput/Source/GameInputBase/Private/IGameInputDeviceInterface.cpp
--- a/Engine/Plugins/Runtime/GameInput/Source/GameInputBase/Private/IGameInputDeviceInterface.cpp
+++ b/Engine/Plugins/Runtime/GameInput/Source/GameInputBase/Private/IGameInputDeviceInterface.cpp
@@ -93,8 +93,7 @@ void IGameInputDeviceInterface::SendControllerEvents()
 		// A map of Platform users to a bitmask of any reading kinds that were processed this frame.
 		// This is used by the devices to keep track of which users had which readings, and use that
 		// state to determine if we can process a given reading.
-		static TMap<FPlatformUserId, GameInputKind> PlatformUsersWhoHaveHadInputThisFrame;
-		PlatformUsersWhoHaveHadInputThisFrame.Reset();
+		TMap<FPlatformUserId, GameInputKind> PlatformUsersWhoHaveHadInputThisFrame;
 
 		// The allowed input kinds that we can read from the IGameInput interface.
 		const GameInputKind AllowedInputKindsThisFrame = GetCurrentGameInputKindSupport();
@@ -291,18 +290,11 @@ GameInputKind IGameInputDeviceInterface::GetCurrentGameInputKindSupport() const
 		RegisterInputKindMask |= GameInputKindArcadeStick;
 	}
 
-	// TODO: Future expansion of GameInput devices!
-	/*if (Settings->bProcessArcadeStick)
+	if (PlatformSettings->bProcessFlightStick)
 	{
-
+		RegisterInputKindMask |= GameInputKindFlightStick;
 	}
 
-	if (Settings->bProcessRacingWheel)
-	{
-
-	}	
-	*/
-
 	return RegisterInputKindMask;
 }
 
diff --git a/Engine/Plugins/Runtime/GameInput/Source/GameInputBase/Public/GameInputDeveloperSettings.h b/Engine/Plugins/Runtime/GameInput/Source/GameInputBase/Public/GameInputDeveloperSettings.h
--- a/Engine/Plugins/Runtime/GameInput/Source/GameInputBase/Public/GameInputDeveloperSettings.h
+++ b/Engine/Plugins/Runtime/GameInput/Source/GameInputBase/Public/GameInputDeveloperSettings.h
@@ -389,6 +389,35 @@ public:
 	*/
 	UPROPERTY(EditAnywhere, Config, Category = "Processing Options", meta=(ConfigRestartRequired = true, DisplayName="Process Arcade Sticks (Experimental)"))
 	bool bProcessArcadeStick = false;
+
+	/**
+	 * If true, then we will process the GameInputKindFlightStick type.
+	 *
+	 * Note: This is experimental!
+	 *
+	 * Default: False
+	 */
+	UPROPERTY(EditAnywhere, Config, Category = "Processing Options", meta=(ConfigRestartRequired = true, DisplayName="Process Flight Sticks (Experimental)"))
+	bool bProcessFlightStick = false;
+
+	/** The default value of any flight stick deadzone values. */
+	static constexpr float DefaultFlightStickDeadzone = (7849.0f / 32768.0f);
+
+	/** Deadzone which should be used for Flight Stick Pitch inputs */
+	UPROPERTY(EditAnywhere, Config, Category = "Device Settings|Flight Stick", meta=(EditCondition="bProcessFlightStick"))
+	float FlightStickPitchDeadzone = DefaultFlightStickDeadzone;
+
+	/** Deadzone which should be used for Flight Stick Roll inputs */
+	UPROPERTY(EditAnywhere, Config, Category = "Device Settings|Flight Stick", meta=(EditCondition="bProcessFlightStick"))
+	float FlightStickRollDeadzone = DefaultFlightStickDeadzone;
+
+	/** Deadzone which should be used for Flight Stick Throttle inputs */
+	UPROPERTY(EditAnywhere, Config, Category = "Device Settings|Flight Stick", meta=(EditCondition="bProcessFlightStick"))
+	float FlightStickThrottleDeadzone = DefaultFlightStickDeadzone;
+
+	/** Deadzone which should be used for Flight Stick Yaw inputs */
+	UPROPERTY(EditAnywhere, Config, Category = "Device Settings|Flight Stick", meta=(EditCondition="bProcessFlightStick"))
+	float FlightStickYawDeadzone = DefaultFlightStickDeadzone;
 };
 
 /**
diff --git a/Engine/Plugins/Runtime/GameInput/Source/GameInputBase/Public/GameInputDeviceProcessor.h b/Engine/Plugins/Runtime/GameInput/Source/GameInputBase/Public/GameInputDeviceProcessor.h
--- a/Engine/Plugins/Runtime/GameInput/Source/GameInputBase/Public/GameInputDeviceProcessor.h
+++ b/Engine/Plugins/Runtime/GameInput/Source/GameInputBase/Public/GameInputDeviceProcessor.h
@@ -523,4 +523,47 @@ protected:
 	double RepeatTime[MaxSupportedButtons];
 };
 
+/**
+ *  Processor for the GameInputKindFlightStick type.
+ */
+class GAMEINPUTBASE_API FGameInputFlightStickProcessor : public IGameInputDeviceProcessor
+{
+public:
+	FGameInputFlightStickProcessor();
+
+protected:
+
+	virtual bool ProcessInput(const FGameInputEventParams& Params) override;
+	virtual bool PostProcessInput(const FGameInputEventParams& Params) override;
+	virtual void ClearState(const FGameInputEventParams& Params) override;
+	virtual GameInputKind GetSupportedReadingKind() const override;
+
+	/**
+	 * Process Flight stick specific buttons.
+	 */
+	bool ProcessFlightStickButtons(const FGameInputEventParams& Params, GameInputFlightStickState& State);
+
+	/**
+	 * Process the analog values (yaw, pitch, roll, and throttle) of the flight stick
+	 */
+	bool ProcessFlightStickAnalog(const FGameInputEventParams& Params, GameInputFlightStickState& State);
+	
+	/**
+	* Keeps track of how many times this gamepad has been processed this frame.
+	* Every successful processing of button input in ProcessInput will increment this value.
+	* This will get reset at the end of the input frame in PostProcessInput.
+	*/
+	int32 NumReadingsProcessedThisFrame = 0;
+	
+	/** State of the flight stick from the previous frame */
+	GameInputFlightStickState PreviousState = {};
+
+	// Repeat times for any pressed buttons
+	static constexpr uint32 MaxSupportedButtons = 32;
+	double RepeatTime[MaxSupportedButtons];
+
+	/** Repeat times for when switches are pressed (aka the DPad) */
+	TArray<double> SwitchRepeatTimes;
+};
+
 #endif	// GAME_INPUT_SUPPORT
\ No newline at end of file
