Fix a null dereference introduced by 36516159.

#rnx

diff --git a/Engine/Plugins/Runtime/GameInput/Source/GameInputBase/Private/GameInputUtils.cpp b/Engine/Plugins/Runtime/GameInput/Source/GameInputBase/Private/GameInputUtils.cpp
--- a/Engine/Plugins/Runtime/GameInput/Source/GameInputBase/Private/GameInputUtils.cpp
+++ b/Engine/Plugins/Runtime/GameInput/Source/GameInputBase/Private/GameInputUtils.cpp
@@ -14,11 +14,15 @@ namespace UE::GameInput
 		{
 			const GameInputDeviceInfo* Info = Device->GetDeviceInfo();
 			const bool bIsVirtualDevice = Info->deviceFamily == GameInputFamilyVirtual;
+
+			const FString DeviceDisplayName =
+				Info->displayName ?
+				FString::Printf(TEXT("%.*hs"), Info->displayName->sizeInBytes, Info->displayName->data) :
+				TEXT("Unknown");
 			
-			Result = FString::Printf(TEXT("%p   DeviceName: %.*hs %s ProdId: %04x (%u) VendId: %04x (%u)"),
+			Result = FString::Printf(TEXT("%p   DeviceName: %s %s ProdId: %04x (%u) VendId: %04x (%u)"),
 				Device, 
-				Info->displayName->sizeInBytes,
-				Info->displayName->data,
+				*DeviceDisplayName,
 				bIsVirtualDevice ? TEXT("(virtual)") : TEXT(""),
 				Info->productId, (uint32)Info->productId,
 				Info->vendorId, (uint32)Info->vendorId);
