Add ability for unattended applications to receive input via commandline param -UnattendedInput.
Solves some local scenarios where access to command prompt and other input is desired.

diff --git a/Engine/Plugins/Runtime/GameInput/Source/GameInputBase/Private/GameInputBaseModule.cpp b/Engine/Plugins/Runtime/GameInput/Source/GameInputBase/Private/GameInputBaseModule.cpp
--- a/Engine/Plugins/Runtime/GameInput/Source/GameInputBase/Private/GameInputBaseModule.cpp
+++ b/Engine/Plugins/Runtime/GameInput/Source/GameInputBase/Private/GameInputBaseModule.cpp
@@ -65,7 +65,7 @@ void FGameInputBaseModule::StartupModule()
 	}
 
 	// Unattended app can't receive any user input, so there is no need to try and create the GameInput interface.
-	if (FApp::IsUnattended())
+	if (FApp::IsUnattended() && !FApp::AllowUnattendedInput())
 	{
 		UE_LOG(LogGameInput, Log, TEXT("GameInputBase module is exiting because it is unattended (FApp::IsUnattended is true) and thus cannot recieve user input. GameInput will not be initalized."));
 		return;
diff --git a/Engine/Source/Runtime/Core/Private/Misc/App.cpp b/Engine/Source/Runtime/Core/Private/Misc/App.cpp
--- a/Engine/Source/Runtime/Core/Private/Misc/App.cpp
+++ b/Engine/Source/Runtime/Core/Private/Misc/App.cpp
@@ -280,6 +280,14 @@ bool FApp::IsUnattended()
 	return bIsUnattended || GIsAutomationTesting;
 }
 
+bool FApp::AllowUnattendedInput()
+{
+	// FCommandLine::Get() will assert that the command line has been set.
+	// This function may not be used before FCommandLine::Set() is called.
+	static bool bUnattendedInput = FParse::Param(FCommandLine::Get(), TEXT("UnattendedInput"));
+	return bUnattendedInput;
+}
+
 bool FApp::ShouldUseThreadingForPerformance()
 {
 	static bool OnlyOneThread = 
diff --git a/Engine/Source/Runtime/Core/Public/Misc/App.h b/Engine/Source/Runtime/Core/Public/Misc/App.h
--- a/Engine/Source/Runtime/Core/Public/Misc/App.h
+++ b/Engine/Source/Runtime/Core/Public/Misc/App.h
@@ -474,6 +474,16 @@ public:
 	 */
 	static CORE_API bool IsUnattended();
 
+	/**
+	 * Checks whether this application allows unattended input.
+	 *
+	 * Unattended applications usually are unable to receive user input.
+	 * This method can be used to determine whether to allow input modules to initialize when unattended.
+	 *
+	 * @return true if the application allows unattended input, false otherwise.
+	 */
+	static CORE_API bool AllowUnattendedInput();
+
 	/**
 	 * Checks whether the application should run multi-threaded for performance critical features.
 	 *
