Cursor: Unified the way mouse cursor movement is handled across console platforms.

diff --git a/Engine/Plugins/Runtime/GameInput/Source/GameInputBase/Private/GameInputDeviceProcessor.cpp b/Engine/Plugins/Runtime/GameInput/Source/GameInputBase/Private/GameInputDeviceProcessor.cpp
--- a/Engine/Plugins/Runtime/GameInput/Source/GameInputBase/Private/GameInputDeviceProcessor.cpp
+++ b/Engine/Plugins/Runtime/GameInput/Source/GameInputBase/Private/GameInputDeviceProcessor.cpp
@@ -8,6 +8,7 @@
 
 
 #include "GenericPlatform/GenericPlatformInputDeviceMapper.h"
+#include "GenericPlatform/CursorUtils.h"
 #include "Framework/Application/SlateApplication.h"
 #include "HAL/IConsoleManager.h"
 #include "GameInputDeveloperSettings.h"
@@ -1124,14 +1125,6 @@ FGameInputMouseDeviceProcessor::FGameInputMouseDeviceProcessor(const TSharedPtr<
 	}
 }
 
-float FGameInputMouseDeviceProcessor::CalculateMouseAcceleration(int32 Delta)
-{
-	const float NominalMovement = 20.0f;
-	const float Sign = (Delta > 0) ? 1.0f : -1.0f;
-	const float FDelta = (float)Delta;
-	return Sign * FMath::Max(FMath::Abs(FDelta), FMath::Abs(GMouseSensitivity * FDelta * FDelta / NominalMovement));
-}
-
 bool FGameInputMouseDeviceProcessor::CanProcessVirtualMouse() const
 {
 	// ignore the input if requested
@@ -1169,15 +1162,19 @@ bool FGameInputMouseDeviceProcessor::ProcessInput(const FGameInputEventParams& P
 
 	// update mouse position
 	float MouseDX, MouseDY = 0.0f;
+	
+	int32 RawMouseDX = static_cast<int32>(MouseState.positionX) - static_cast<int32>(PreviousMouseState.positionX);
+	int32 RawMouseDY = static_cast<int32>(MouseState.positionY) - static_cast<int32>(PreviousMouseState.positionY);	
+	
 	if (bHighPrecisionMouseMode)
 	{
-		MouseDX = static_cast<float>(MouseState.positionX - PreviousMouseState.positionX);
-		MouseDY = static_cast<float>(MouseState.positionY - PreviousMouseState.positionY);
+		MouseDX = static_cast<float>(RawMouseDX);
+		MouseDY = static_cast<float>(RawMouseDY);
 	}
 	else
 	{
-		MouseDX = CalculateMouseAcceleration(static_cast<int32>(MouseState.positionX) - static_cast<int32>(PreviousMouseState.positionX));
-		MouseDY = CalculateMouseAcceleration(static_cast<int32>(MouseState.positionY) - static_cast<int32>(PreviousMouseState.positionY));
+		MouseDX = UE::Cursor::CalculateDeltaWithAcceleration(RawMouseDX, GMouseSensitivity);
+		MouseDY = UE::Cursor::CalculateDeltaWithAcceleration(RawMouseDY, GMouseSensitivity);
 	}
 	if (MouseDX != 0 || MouseDY != 0)
 	{
@@ -1189,11 +1186,14 @@ bool FGameInputMouseDeviceProcessor::ProcessInput(const FGameInputEventParams& P
 			Cursor->SetPosition((int32)CursorPos.X, (int32)CursorPos.Y);
 		}
 
+		// LastMouseOffset is used for double-click detection, so we should use the same coordinates as the cursor i.e. the processed ones.
 		LastMouseOffset.X += MouseDX;
 		LastMouseOffset.Y += MouseDY;
 
-		UE_LOG(LogGameInput, VeryVerbose, TEXT("Device %s (InputDeviceId = %d) - Mouse DX %0.2f, DY %.2f  %s"), *UE::GameInput::LexToString(Params.Device), Params.InputDeviceId.GetId(), MouseDX, MouseDY, bHighPrecisionMouseMode ? TEXT("(high precision)") : TEXT(""));
-		Params.MessageHandler->OnRawMouseMove((int32)MouseDX, (int32)MouseDY);
+		UE_LOG(LogGameInput, VeryVerbose, TEXT("Device %s (InputDeviceId = %d) - Mouse RawDX %d, RawDY %d, DX %0.2f, DY %.2f  %s"), *UE::GameInput::LexToString(Params.Device), Params.InputDeviceId.GetId(), RawMouseDX, RawMouseDY, MouseDX, MouseDY, bHighPrecisionMouseMode ? TEXT("(high precision)") : TEXT(""));
+
+		// OnRawMouseMove may be used to process player character controls so, as the name suggests, these should be raw data without any smoothing.
+		Params.MessageHandler->OnRawMouseMove(RawMouseDX, RawMouseDY);
 	}
 
 	// update mouse wheel
diff --git a/Engine/Plugins/Runtime/GameInput/Source/GameInputBase/Public/GameInputDeviceProcessor.h b/Engine/Plugins/Runtime/GameInput/Source/GameInputBase/Public/GameInputDeviceProcessor.h
--- a/Engine/Plugins/Runtime/GameInput/Source/GameInputBase/Public/GameInputDeviceProcessor.h
+++ b/Engine/Plugins/Runtime/GameInput/Source/GameInputBase/Public/GameInputDeviceProcessor.h
@@ -282,8 +282,6 @@ protected:
 	virtual void ClearState(const FGameInputEventParams& Params) override;
 	virtual GameInputKind GetSupportedReadingKind() const override;
 
-	float CalculateMouseAcceleration(int32 Delta);
-
 	virtual bool CanProcessVirtualMouse() const;
 
 	TSharedPtr<class ICursor> Cursor;
diff --git a/Engine/Source/Runtime/ApplicationCore/Internal/GenericPlatform/CursorUtils.h b/Engine/Source/Runtime/ApplicationCore/Internal/GenericPlatform/CursorUtils.h
new file mode 100644
--- /dev/null
+++ b/Engine/Source/Runtime/ApplicationCore/Internal/GenericPlatform/CursorUtils.h
@@ -0,0 +1,11 @@
+// Copyright Epic Games, Inc. All Rights Reserved.
+
+#pragma once
+
+#include "CoreTypes.h"
+
+
+namespace UE::Cursor
+{
+	APPLICATIONCORE_API float CalculateDeltaWithAcceleration(int32 Delta, float Sensitivity);
+}
diff --git a/Engine/Source/Runtime/ApplicationCore/Private/GenericPlatform/CursorUtils.cpp b/Engine/Source/Runtime/ApplicationCore/Private/GenericPlatform/CursorUtils.cpp
new file mode 100644
--- /dev/null
+++ b/Engine/Source/Runtime/ApplicationCore/Private/GenericPlatform/CursorUtils.cpp
@@ -0,0 +1,17 @@
+// Copyright Epic Games, Inc. All Rights Reserved.
+
+#include "GenericPlatform/CursorUtils.h"
+#include "CoreMinimal.h"
+
+
+namespace UE::Cursor
+{
+	float CalculateDeltaWithAcceleration(int32 Delta, float Sensitivity)
+	{
+		constexpr float NominalMovement = 20.0f;
+		const float Sign = (Delta > 0) ? 1.0f : -1.0f;
+		const float FDelta = (float)Delta;
+
+		return Sign * FMath::Max(FMath::Abs(FDelta), FMath::Abs(Sensitivity * FDelta * FDelta / NominalMovement));
+	}
+}
