Fix for GameInputBase module if the GDK isn't available

diff --git a/Engine/Plugins/Runtime/GameInput/Source/GameInputBase/GameInputBase.build.cs b/Engine/Plugins/Runtime/GameInput/Source/GameInputBase/GameInputBase.build.cs
--- a/Engine/Plugins/Runtime/GameInput/Source/GameInputBase/GameInputBase.build.cs
+++ b/Engine/Plugins/Runtime/GameInput/Source/GameInputBase/GameInputBase.build.cs
@@ -107,13 +107,7 @@ namespace UnrealBuildTool.Rules
 		 */
 		private string GetGameInputVersionString(in string GRDKPath)
 		{
-			string GRDKVersionString = GetGameInputVersionStringFromIni(GRDKPath);
-			if (GRDKVersionString == "-1")
-			{
-				GRDKVersionString = GetGameInputVersionStringFromPath(GRDKPath);
-			}
-
-			return GRDKVersionString;
+			return GetGameInputVersionStringFromIni(GRDKPath);
 		}
 		
 		/**
@@ -153,29 +147,5 @@ namespace UnrealBuildTool.Rules
 
 			return "-1";
 		}
-
-		/**
-		 * Returns the 6 digit version number based on the installation path of the SDK.
-		 * -1 if the path is invalid or a valid version couldn't be found.
-		 */
-		private string GetGameInputVersionStringFromPath(in string GRDKPath)
-		{
-			if (GRDKPath is null)
-			{
-				return "-1";
-			}
-
-			// regex for finding a 6 digit number before "GRDK", which should always give us the version number of the GRDK that is installed
-			string Pattern = "(\\d{6})(\\\\GRDK)";
-			Regex Rg = new Regex(Pattern);
-
-			MatchCollection Matched = Rg.Matches(GRDKPath);
-
-			// The result we care about will be in the first match, second group
-			return
-				Matched.Count > 0 && Matched[0].Groups.Count > 1 ?
-					Matched[0].Groups[1].ToString()
-					: "-1";
-		}
 	}
 }
