Refresh game input device processors when you make changes in the editor. This makes it much easier to iterate and tweak game input settings.

#rnx

diff --git a/Engine/Plugins/Experimental/GameInputWindows/Source/GameInputWindows/Private/GameInputWindowsDevice.cpp b/Engine/Plugins/Experimental/GameInputWindows/Source/GameInputWindows/Private/GameInputWindowsDevice.cpp
--- a/Engine/Plugins/Experimental/GameInputWindows/Source/GameInputWindows/Private/GameInputWindowsDevice.cpp
+++ b/Engine/Plugins/Experimental/GameInputWindows/Source/GameInputWindows/Private/GameInputWindowsDevice.cpp
@@ -13,12 +13,16 @@
 FGameInputWindowsInputDevice::FGameInputWindowsInputDevice(const TSharedRef<FGenericApplicationMessageHandler>& InMessageHandler, IGameInput* InGameInput)
 	: IGameInputDeviceInterface(InMessageHandler, InGameInput)
 {
-	
+#if WITH_EDITOR
+	SetupEditorSettingListener();
+#endif
 }
 
 FGameInputWindowsInputDevice::~FGameInputWindowsInputDevice()
 {
-	
+#if WITH_EDITOR
+	CleanupEditorSettingListener();
+#endif
 }
 
 GameInputKind FGameInputWindowsInputDevice::GetCurrentGameInputKindSupport() const
@@ -155,4 +159,52 @@ FGameInputDeviceContainer* FGameInputWindowsInputDevice::CreateDeviceData(IGameI
 	return Container.Get();
 }
 
+#if WITH_EDITOR
+
+void FGameInputWindowsInputDevice::SetupEditorSettingListener()
+{
+	UGameInputDeveloperSettings* Settings = GetMutableDefault<UGameInputDeveloperSettings>();
+	if (!Settings)
+	{
+		return;
+	}
+
+	EditorSettingChangedDelegate = Settings->OnInputSettingChanged.AddRaw(this, &FGameInputWindowsInputDevice::HandleEditorSettingChanged);	
+}
+
+void FGameInputWindowsInputDevice::CleanupEditorSettingListener()
+{
+	if (!EditorSettingChangedDelegate.IsValid())
+	{
+		return;
+	}
+
+	UGameInputDeveloperSettings* Settings = GetMutableDefault<UGameInputDeveloperSettings>();
+	if (!Settings)
+	{
+		return;
+	}
+	
+	Settings->OnInputSettingChanged.Remove(EditorSettingChangedDelegate);
+}
+
+void FGameInputWindowsInputDevice::HandleEditorSettingChanged()
+{
+	for (TSharedPtr<FGameInputDeviceContainer> Device : DeviceData)
+	{
+		// Recreate the device processors for every device
+		Device->RecreateDeviceProcessors(GetCurrentGameInputKindSupport());
+
+		// If this device previously had an invalid DeviceId because it had no processors, but now it does,
+		// then we need to assign a FInputDeviceId because it can now process input
+		// We will never change a device id to be invalid though, because that could cause weirdness where Slate
+		// is expecting input events from a device which no longer exists.
+		if (!Device->GetDeviceId().IsValid() && Device->GetNumberOfProcessors() > 0)
+		{
+			Device->SetInputDeviceId(IPlatformInputDeviceMapper::Get().AllocateNewInputDeviceId());
+		}
+	}
+}
+#endif	// WITH_EDITOR
+
 #endif	//#if GAME_INPUT_SUPPORT
\ No newline at end of file
diff --git a/Engine/Plugins/Experimental/GameInputWindows/Source/GameInputWindows/Private/GameInputWindowsDevice.h b/Engine/Plugins/Experimental/GameInputWindows/Source/GameInputWindows/Private/GameInputWindowsDevice.h
--- a/Engine/Plugins/Experimental/GameInputWindows/Source/GameInputWindows/Private/GameInputWindowsDevice.h
+++ b/Engine/Plugins/Experimental/GameInputWindows/Source/GameInputWindows/Private/GameInputWindowsDevice.h
@@ -20,6 +20,15 @@ protected:
 	virtual void HandleDeviceConnected(IGameInputDevice* Device, uint64 Timestamp) override;
 	virtual FGameInputDeviceContainer* CreateDeviceData(IGameInputDevice* InDevice) override;
 	//~End IGameInputDeviceInterface interface
+	
+private:
+#if WITH_EDITOR
+	void SetupEditorSettingListener();
+	void CleanupEditorSettingListener();
+	void HandleEditorSettingChanged();
+
+	FDelegateHandle EditorSettingChangedDelegate;
+#endif
 };
 
 #endif	// #if GAME_INPUT_SUPPORT
\ No newline at end of file
diff --git a/Engine/Plugins/Runtime/GameInput/Source/GameInputBase/Private/GameInputDeveloperSettings.cpp b/Engine/Plugins/Runtime/GameInput/Source/GameInputBase/Private/GameInputDeveloperSettings.cpp
--- a/Engine/Plugins/Runtime/GameInput/Source/GameInputBase/Private/GameInputDeveloperSettings.cpp
+++ b/Engine/Plugins/Runtime/GameInput/Source/GameInputBase/Private/GameInputDeveloperSettings.cpp
@@ -126,20 +126,6 @@ UGameInputDeveloperSettings::UGameInputDeveloperSettings(const FObjectInitialize
 }
 
 #if WITH_EDITOR
-
-namespace UE::GameInput
-{
-	/** A set of FName's that, if change in PostEditChangeProperty, would require Game Input to update what device status updates it is listening for. */
-	static const TSet<FName>& GetPropertyNamesThatRequireGameInputUpdate()
-	{
-		static const TSet<FName> PropertyNamesThatRequireGameInputUpdate
-		{
-			GET_MEMBER_NAME_CHECKED(UGameInputPlatformSettings, bProcessController),
-		};
-		return PropertyNamesThatRequireGameInputUpdate;
-	}
-}
-
 const TArray<uint32>& UGameInputDeveloperSettings::GetControllerButtonMappingDataKeyOptions()
 {
 	// The controller button bitmask is generated by bit shifting the current index
@@ -183,16 +169,11 @@ const TArray<uint32>& UGameInputDeveloperSettings::GetControllerButtonMappingDat
 
 void UGameInputDeveloperSettings::PostEditChangeProperty(FPropertyChangedEvent& PropertyChangedEvent)
 {
-	Super::PostEditChangeProperty(PropertyChangedEvent);
-
-	const FName PropertyName = (PropertyChangedEvent.Property != nullptr) ? PropertyChangedEvent.Property->GetFName() : NAME_None;
+	Super::PostEditChangeProperty(PropertyChangedEvent);	
 
-	if (UE::GameInput::GetPropertyNamesThatRequireGameInputUpdate().Contains(PropertyName))
+	if (OnInputSettingChanged.IsBound())
 	{
-		if (OnInputSettingChanged.IsBound())
-		{
-			OnInputSettingChanged.Execute(PropertyName);
-		}
+		OnInputSettingChanged.Broadcast();
 	}
 }
 #endif	// WITH_EDITOR
\ No newline at end of file
diff --git a/Engine/Plugins/Runtime/GameInput/Source/GameInputBase/Private/GameInputDeviceContainer.cpp b/Engine/Plugins/Runtime/GameInput/Source/GameInputBase/Private/GameInputDeviceContainer.cpp
--- a/Engine/Plugins/Runtime/GameInput/Source/GameInputBase/Private/GameInputDeviceContainer.cpp
+++ b/Engine/Plugins/Runtime/GameInput/Source/GameInputBase/Private/GameInputDeviceContainer.cpp
@@ -69,6 +69,22 @@ void FGameInputDeviceContainer::InitalizeDeviceProcessors()
 	InitalizeDeviceProcessors_Impl();
 }
 
+void FGameInputDeviceContainer::RecreateDeviceProcessors(const GameInputKind InAllowedGameInputKinds)
+{
+	UE_CLOG(AllowedGameInputKinds != InAllowedGameInputKinds,
+		LogGameInput, Log, TEXT("[%hs] Recreating device processors with new allowance. Changing from %s to %s"),
+		__func__, *UE::GameInput::LexToString(AllowedGameInputKinds), *UE::GameInput::LexToString(InAllowedGameInputKinds));
+	
+	AllowedGameInputKinds = InAllowedGameInputKinds;
+
+	// Tear down the old input processors
+	ClearInputState(nullptr);
+	Processors.Empty();
+
+	// Create them again with the new allowed input kinds
+	InitalizeDeviceProcessors();
+}
+
 void FGameInputDeviceContainer::InitalizeDeviceProcessors_Impl()
 {
 	if (!Device)
diff --git a/Engine/Plugins/Runtime/GameInput/Source/GameInputBase/Public/GameInputDeveloperSettings.h b/Engine/Plugins/Runtime/GameInput/Source/GameInputBase/Public/GameInputDeveloperSettings.h
--- a/Engine/Plugins/Runtime/GameInput/Source/GameInputBase/Public/GameInputDeveloperSettings.h
+++ b/Engine/Plugins/Runtime/GameInput/Source/GameInputBase/Public/GameInputDeveloperSettings.h
@@ -302,7 +302,7 @@ public:
 	 *
 	 * Default: False
 	 */
-	UPROPERTY(Config, EditAnywhere, Category="Processing Options|Require Configuration", meta = (ConfigRestartRequired = true))
+	UPROPERTY(Config, EditAnywhere, Category="Processing Options|Require Configuration")
 	bool bProcessController = false;
 
 	/** 
@@ -310,7 +310,7 @@ public:
 	*
 	* Default: False
 	*/
-	UPROPERTY(Config, EditAnywhere, Category="Processing Options|Require Configuration", meta = (ConfigRestartRequired = true))
+	UPROPERTY(Config, EditAnywhere, Category="Processing Options|Require Configuration")
 	bool bProcessRawInput = false;
 	
 	/**
@@ -319,7 +319,7 @@ public:
 	* 
 	* Default: True
 	*/
-	UPROPERTY(Config, EditAnywhere, Category="Processing Options|Require Configuration", meta=(EditCondition="bProcessController || bProcessRawInput", ConfigRestartRequired = true))
+	UPROPERTY(Config, EditAnywhere, Category="Processing Options|Require Configuration", meta=(EditCondition="bProcessController || bProcessRawInput"))
 	bool bSpecialDevicesRequireExplicitDeviceConfiguration = true;
 
 	/** 
@@ -330,7 +330,7 @@ public:
 	* 
 	* Default: True
 	*/
-	UPROPERTY(Config, EditAnywhere, Category="Processing Options", meta = (ConfigRestartRequired = true))
+	UPROPERTY(Config, EditAnywhere, Category="Processing Options")
 	bool bProcessGamepad = true;
 
 	/** 
@@ -341,7 +341,7 @@ public:
 	* 
 	* Default: True
 	*/
-	UPROPERTY(Config, EditAnywhere, Category="Processing Options", meta = (ConfigRestartRequired = true))
+	UPROPERTY(Config, EditAnywhere, Category="Processing Options")
 	bool bProcessKeyboard = true;
 
 	/** 
@@ -352,7 +352,7 @@ public:
 	* 
 	* Default: True
 	*/
-	UPROPERTY(Config, EditAnywhere, Category="Processing Options", meta = (ConfigRestartRequired = true))
+	UPROPERTY(Config, EditAnywhere, Category="Processing Options")
 	bool bProcessMouse = true;
 
 	/**
@@ -366,7 +366,7 @@ public:
 	 * 
 	 * Default: False
 	 */
-	UPROPERTY(EditAnywhere, Config, Category = "Processing Options|Require Configuration", meta=(ConfigRestartRequired = true, DisplayName="Process Racing Wheel (Experimental)"))
+	UPROPERTY(EditAnywhere, Config, Category = "Processing Options|Require Configuration", meta=(DisplayName="Process Racing Wheel (Experimental)"))
 	bool bProcessRacingWheel = false;
 
 	/** The default racing wheel deadzone */
@@ -387,7 +387,7 @@ public:
 	* 
 	* Default: False
 	*/
-	UPROPERTY(EditAnywhere, Config, Category = "Processing Options", meta=(ConfigRestartRequired = true, DisplayName="Process Arcade Sticks (Experimental)"))
+	UPROPERTY(EditAnywhere, Config, Category = "Processing Options", meta=(DisplayName="Process Arcade Sticks (Experimental)"))
 	bool bProcessArcadeStick = false;
 
 	/**
@@ -397,7 +397,7 @@ public:
 	 *
 	 * Default: False
 	 */
-	UPROPERTY(EditAnywhere, Config, Category = "Processing Options", meta=(ConfigRestartRequired = true, DisplayName="Process Flight Sticks (Experimental)"))
+	UPROPERTY(EditAnywhere, Config, Category = "Processing Options", meta=(DisplayName="Process Flight Sticks (Experimental)"))
 	bool bProcessFlightStick = false;
 
 	/** The default value of any flight stick deadzone values. */
@@ -495,9 +495,9 @@ public:
 
 #if WITH_EDITOR
 public:
-	DECLARE_DELEGATE_OneParam(FGameInputSettingChanged, FName);
+
 	/** A delegate that is called on PostEditChangeProperty with the name of the property that has changed. */
-	FGameInputSettingChanged OnInputSettingChanged;
+	TMulticastDelegate<void()> OnInputSettingChanged;
 	
 	/**
 	 * Returns all known uint32 options for controller button configurations.
diff --git a/Engine/Plugins/Runtime/GameInput/Source/GameInputBase/Public/GameInputDeviceContainer.h b/Engine/Plugins/Runtime/GameInput/Source/GameInputBase/Public/GameInputDeviceContainer.h
--- a/Engine/Plugins/Runtime/GameInput/Source/GameInputBase/Public/GameInputDeviceContainer.h
+++ b/Engine/Plugins/Runtime/GameInput/Source/GameInputBase/Public/GameInputDeviceContainer.h
@@ -31,6 +31,12 @@ public:
 	/** Create any device processors necessary for this Game Input device to process input and send input events */
 	void InitalizeDeviceProcessors();
 
+	/**
+	 * Clears the state of any device processors and destroyies them. Empties the processors array.
+	 * Sets the allowed input kind to the new given input kinds
+	 */
+	void RecreateDeviceProcessors(const GameInputKind InAllowedGameInputKinds);
+
 	/** Set the message handler that this Device Container should send events through. */
 	void SetMessageHandler(const TSharedRef<FGenericApplicationMessageHandler>& InMessageHandler);
 	
