IGameInputDeviceInterface

Call "ProcessDeferredDevice" connections inside Tick since Tick happens on the game thread and SendControllerEvents might not.

#tests Harmonix Dev
#rnx

diff --git a/Engine/Plugins/Runtime/GameInput/Source/GameInputBase/Private/IGameInputDeviceInterface.cpp b/Engine/Plugins/Runtime/GameInput/Source/GameInputBase/Private/IGameInputDeviceInterface.cpp
--- a/Engine/Plugins/Runtime/GameInput/Source/GameInputBase/Private/IGameInputDeviceInterface.cpp
+++ b/Engine/Plugins/Runtime/GameInput/Source/GameInputBase/Private/IGameInputDeviceInterface.cpp
@@ -63,7 +63,21 @@ void IGameInputDeviceInterface::Initialize()
 
 void IGameInputDeviceInterface::Tick(float DeltaTime)
 {
-	// Required by IInputDevice
+#if GAME_INPUT_SUPPORT
+
+	FScopeLock Lock(&DeviceInfoCS);
+
+	// At this point we should always have a valid Game Input object
+	check(GameInput);
+	
+	// Process any input devices so long as we are not constrained
+	if (!bIsAppCurrentlyConstrained)
+	{
+		// Handle any device connection/disconnection state changes first before attempting to process any devices
+		ProcessDeferredDeviceConnectionChanges();
+	}
+	
+#endif // GAME_INPUT_SUPPORT
 }
 
 void IGameInputDeviceInterface::SendControllerEvents()
@@ -87,9 +101,6 @@ void IGameInputDeviceInterface::SendControllerEvents()
 	// Process any input devices so long as we are not constrained
 	if (!bIsAppCurrentlyConstrained)
 	{
-		// Handle any device connection/disconnection state changes first before attempting to process any devices
-		ProcessDeferredDeviceConnectionChanges();
-		
 		// A map of Platform users to a bitmask of any reading kinds that were processed this frame.
 		// This is used by the devices to keep track of which users had which readings, and use that
 		// state to determine if we can process a given reading.
