Move the evaluation of the digital FKeys associated with gamepad thumbsticks and triggers in Game Input to only be called once, where we evaluate the actual analog values.

This fixes a bug where, when there are multiple readings for a single input device in a given frame, the digital FKey's (i.e. GamepadButtonAux_LeftTrigger) would be fired, but the associated analog value would be incorrect.

Also, make sure that the analog events are processed BEFORE the button states when they are evaluated on the same frame, making game input more consistent with other platforms.

diff --git a/Engine/Plugins/Runtime/GameInput/Source/GameInputBase/Private/GameInputDeviceProcessor.cpp b/Engine/Plugins/Runtime/GameInput/Source/GameInputBase/Private/GameInputDeviceProcessor.cpp
--- a/Engine/Plugins/Runtime/GameInput/Source/GameInputBase/Private/GameInputDeviceProcessor.cpp
+++ b/Engine/Plugins/Runtime/GameInput/Source/GameInputBase/Private/GameInputDeviceProcessor.cpp
@@ -388,7 +388,12 @@
 	}
 
 	bool bRes = false;
-
+	
+	// We only want to process gamepad ANALOG inputs on the last reading because we only care about
+	// the most recent analog stick input value. If there are multiple readings and we processed analog
+	// for every one, then the values would accumulate and stack, giving us incorrect data in the message handler.
+	bRes |= ProcessGamepadAnalogState(Params, GamepadState);
+	
 	// If there were no gamepad events this frame, send button events using the previous frame's reading for button repeats.
 	// This is necessary because GetNextReading won't return any reading if the state is unchanged
 	if (!bProcessedAnyGamepadButtonsThisFrame)
@@ -396,11 +401,6 @@
 		bRes |= ProcessGamepadButtonState(Params, GamepadState);
 	}
 
-	// We only want to process gamepad ANALOG inputs on the last reading because we only care about
-	// the most recent analog stick input value. If there are multiple readings and we processed analog
-	// for every one, then the values would accumulate and stack, giving us incorrect data in the message handler.
-	bRes |= ProcessGamepadAnalogState(Params, GamepadState);
-
 	return bRes;
 }
 
@@ -426,7 +426,63 @@
 	OnControllerAnalog(Params, FGamepadKeyNames::RightAnalogY, GamepadState.rightThumbstickY, PreviousState.rightThumbstickY, UE::GameInput::GamepadRightStickDeadzone);
 	OnControllerAnalog(Params, FGamepadKeyNames::LeftTriggerAnalog, GamepadState.leftTrigger, PreviousState.leftTrigger, UE::GameInput::GamepadTriggerDeadzone);
 	OnControllerAnalog(Params, FGamepadKeyNames::RightTriggerAnalog, GamepadState.rightTrigger, PreviousState.rightTrigger, UE::GameInput::GamepadTriggerDeadzone);
+	
+	// map analog triggers to digital input.
+	uint32 AnalogButtonMask = 0u;
+	if (GamepadState.leftTrigger > UE::GameInput::GamepadTriggerDeadzone)
+	{
+		AnalogButtonMask |= UE::GameInput::EGameInputAuxButton::GamepadButtonAux_LeftTrigger;
+	}
 
+	if (GamepadState.rightTrigger > UE::GameInput::GamepadTriggerDeadzone)
+	{
+		AnalogButtonMask |= UE::GameInput::EGameInputAuxButton::GamepadButtonAux_RightTrigger;
+	}
+
+	// map left/right stick digital inputs to (top 8 bits)
+	if (GamepadState.leftThumbstickY > UE::GameInput::GamepadLeftStickDeadzone)
+	{
+		AnalogButtonMask |= UE::GameInput::EGameInputAuxButton::GamepadButtonAux_LeftStickUp;
+	}
+	else if (GamepadState.leftThumbstickY < -UE::GameInput::GamepadLeftStickDeadzone)
+	{
+		AnalogButtonMask |= UE::GameInput::EGameInputAuxButton::GamepadButtonAux_LeftStickDown;
+	}
+	if (GamepadState.leftThumbstickX > UE::GameInput::GamepadLeftStickDeadzone)
+	{
+		AnalogButtonMask |= UE::GameInput::EGameInputAuxButton::GamepadButtonAux_LeftStickRight;
+	}
+	else if (GamepadState.leftThumbstickX < -UE::GameInput::GamepadLeftStickDeadzone)
+	{
+		AnalogButtonMask |= UE::GameInput::EGameInputAuxButton::GamepadButtonAux_LeftStickLeft;
+	}
+
+	if (GamepadState.rightThumbstickY > UE::GameInput::GamepadRightStickDeadzone)
+	{
+		AnalogButtonMask |= UE::GameInput::EGameInputAuxButton::GamepadButtonAux_RightStickUp;
+	}
+	else if (GamepadState.rightThumbstickY < -UE::GameInput::GamepadRightStickDeadzone)
+	{
+		AnalogButtonMask |= UE::GameInput::EGameInputAuxButton::GamepadButtonAux_RightStickDown;
+	}
+	if (GamepadState.rightThumbstickX > UE::GameInput::GamepadRightStickDeadzone)
+	{
+		AnalogButtonMask |= UE::GameInput::EGameInputAuxButton::GamepadButtonAux_RightStickRight;
+	}
+	else if (GamepadState.rightThumbstickX < -UE::GameInput::GamepadRightStickDeadzone)
+	{
+		AnalogButtonMask |= UE::GameInput::EGameInputAuxButton::GamepadButtonAux_RightStickLeft;
+	}
+
+	// Evaluate the analog button mask
+	EvaluateButtonStates(
+		Params,
+		AnalogButtonMask,
+		OUT LastAnalogButtonHeldMask,
+		RepeatTime,
+		UE::GameInput::GetGamepadButtonMap(),
+		MaxSupportedButtons);
+	
 	// update saved analog state
 	PreviousState.leftTrigger = GamepadState.leftTrigger;
 	PreviousState.rightTrigger = GamepadState.rightTrigger;
@@ -458,52 +514,6 @@
 	// map buttons (low 15 bits)
 	uint32 CurrentButtonHeldMask = (static_cast<uint32>(GamepadState.buttons) & UE::GameInput::GamingInputButtonMask);
 
-	// map analog triggers to digital input.
-	if (GamepadState.leftTrigger > UE::GameInput::GamepadTriggerDeadzone)
-	{
-		CurrentButtonHeldMask |= UE::GameInput::EGameInputAuxButton::GamepadButtonAux_LeftTrigger;
-	}
-
-	if (GamepadState.rightTrigger > UE::GameInput::GamepadTriggerDeadzone)
-	{
-		CurrentButtonHeldMask |= UE::GameInput::EGameInputAuxButton::GamepadButtonAux_RightTrigger;
-	}
-
-	// map left/right stick digital inputs to (top 8 bits)
-	if (GamepadState.leftThumbstickY > UE::GameInput::GamepadLeftStickDeadzone)
-	{
-		CurrentButtonHeldMask |= UE::GameInput::EGameInputAuxButton::GamepadButtonAux_LeftStickUp;
-	}
-	else if (GamepadState.leftThumbstickY < -UE::GameInput::GamepadLeftStickDeadzone)
-	{
-		CurrentButtonHeldMask |= UE::GameInput::EGameInputAuxButton::GamepadButtonAux_LeftStickDown;
-	}
-	if (GamepadState.leftThumbstickX > UE::GameInput::GamepadLeftStickDeadzone)
-	{
-		CurrentButtonHeldMask |= UE::GameInput::EGameInputAuxButton::GamepadButtonAux_LeftStickRight;
-	}
-	else if (GamepadState.leftThumbstickX < -UE::GameInput::GamepadLeftStickDeadzone)
-	{
-		CurrentButtonHeldMask |= UE::GameInput::EGameInputAuxButton::GamepadButtonAux_LeftStickLeft;
-	}
-
-	if (GamepadState.rightThumbstickY > UE::GameInput::GamepadRightStickDeadzone)
-	{
-		CurrentButtonHeldMask |= UE::GameInput::EGameInputAuxButton::GamepadButtonAux_RightStickUp;
-	}
-	else if (GamepadState.rightThumbstickY < -UE::GameInput::GamepadRightStickDeadzone)
-	{
-		CurrentButtonHeldMask |= UE::GameInput::EGameInputAuxButton::GamepadButtonAux_RightStickDown;
-	}
-	if (GamepadState.rightThumbstickX > UE::GameInput::GamepadRightStickDeadzone)
-	{
-		CurrentButtonHeldMask |= UE::GameInput::EGameInputAuxButton::GamepadButtonAux_RightStickRight;
-	}
-	else if (GamepadState.rightThumbstickX < -UE::GameInput::GamepadRightStickDeadzone)
-	{
-		CurrentButtonHeldMask |= UE::GameInput::EGameInputAuxButton::GamepadButtonAux_RightStickLeft;
-	}
-
 	EvaluateButtonStates(
 		Params,
 		CurrentButtonHeldMask,
@@ -549,6 +559,9 @@
 		UE::GameInput::GetGamepadButtonMap(),
 		MaxSupportedButtons);
 
+	// Clear the analog button mask
+	LastAnalogButtonHeldMask = 0u;
+
 	for (uint32 i = 0; i < MaxSupportedButtons; i++)
 	{
 		RepeatTime[i] = 0.0;
diff --git a/Engine/Plugins/Runtime/GameInput/Source/GameInputBase/Public/GameInputDeviceProcessor.h b/Engine/Plugins/Runtime/GameInput/Source/GameInputBase/Public/GameInputDeviceProcessor.h
--- a/Engine/Plugins/Runtime/GameInput/Source/GameInputBase/Public/GameInputDeviceProcessor.h
+++ b/Engine/Plugins/Runtime/GameInput/Source/GameInputBase/Public/GameInputDeviceProcessor.h
@@ -168,6 +168,8 @@ protected:
 
 	uint32 LastButtonHeldMask = 0;
 
+	uint32 LastAnalogButtonHeldMask = 0;
+
 	enum { MaxSupportedButtons = 32 };
 
 	double RepeatTime[MaxSupportedButtons];
