Remove unnecessary CVar and TSet's tracking which game input device processors were used this frame.

#rnx

diff --git a/Engine/Plugins/Runtime/GameInput/Source/GameInputBase/Private/GameInputDeviceContainer.cpp b/Engine/Plugins/Runtime/GameInput/Source/GameInputBase/Private/GameInputDeviceContainer.cpp
--- a/Engine/Plugins/Runtime/GameInput/Source/GameInputBase/Private/GameInputDeviceContainer.cpp
+++ b/Engine/Plugins/Runtime/GameInput/Source/GameInputBase/Private/GameInputDeviceContainer.cpp
@@ -10,17 +10,6 @@
 
 #if GAME_INPUT_SUPPORT
 
-namespace UE::GameInput
-{
-	static TAutoConsoleVariable<bool> CVarGameInputForceDeviceProcesorsOncePerFrame
-	(
-		TEXT("gameinput.bForceDeviceProcesorsOncePerFrame"),
-		false,
-		TEXT("If true, then FGameInputDeviceContainer::ProcessInput will ensure that device processors are only called once per ProcessInput call (i.e. once per frame)"),
-		ECVF_Default
-	);
-};
-
 FGameInputDeviceContainer::FGameInputDeviceContainer(
 	const TSharedRef<FGenericApplicationMessageHandler>& InMessageHandler,
 	IGameInputDevice* InDevice,
@@ -210,16 +199,6 @@ const GameInputKind FGameInputDeviceContainer::ProcessInput(IGameInput* GameInpu
 		GameInput->GetCurrentReading(CurrentSupportedKind, Device, &LastReading);
 	}
 
-	// Keep track of the number of input device processors we have used this frame.
-	// We only want to process each one once...
-	static TSet<TSharedPtr<IGameInputDeviceProcessor>> UsedProcessorsThisFrame;
-	UsedProcessorsThisFrame.Reset();
-
-	// Set of processors that we should skip because the platform user has already processed one
-	// of their type this frame already.
-	static TSet<TSharedPtr<IGameInputDeviceProcessor>> ProcessorsToSkipThisFrame;
-	ProcessorsToSkipThisFrame.Reset();
-
 	int32 NumReadingsProcessed = 0;
 
 	while (LastReading.IsValid())
@@ -228,7 +207,7 @@ const GameInputKind FGameInputDeviceContainer::ProcessInput(IGameInput* GameInpu
 		TComPtr<IGameInputReading> Reading;
 		HRESULT hr = GameInput->GetNextReading(LastReading.Get(), CurrentSupportedKind, Device, &Reading);
 		
-		// On the last reading of the frame, there will not be a "next" reading.
+		// On the last reading of the frame for this device, there will not be a "next" reading.
 		const bool bIsLastReadingOfFrame = (hr == GAMEINPUT_E_READING_NOT_FOUND);
 
 		if (FAILED(hr))
@@ -260,26 +239,13 @@ const GameInputKind FGameInputDeviceContainer::ProcessInput(IGameInput* GameInpu
 		Params.PlatformUserId = UserId;
 		Params.InputDeviceId = AssignedDeviceId;
 
-		const bool bEnforceOncePerFrameProcessors = UE::GameInput::CVarGameInputForceDeviceProcesorsOncePerFrame.GetValueOnAnyThread();
-
 		for (TSharedPtr<IGameInputDeviceProcessor>& Processor : Processors)
 		{
-			// Skip any processors that have already been used this frame
-			if (bEnforceOncePerFrameProcessors && UsedProcessorsThisFrame.Contains(Processor))
-			{
-				continue;
-			}
-
 			const GameInputKind ProcessorKind = Processor->GetSupportedReadingKind();
 
 			// The first time we process this, check for if the user has already handled a processor of this kind during this frame.
 			// If they have, then we should skip processing for the rest of this frame.
 			if (NumReadingsProcessed == 0 && bDoNotProcessDuplicateCapabilitiesForSingleUser && (ProcessedKindsForPlatformUserThisFrame & ProcessorKind) != 0)
-			{
-				ProcessorsToSkipThisFrame.Add(Processor);
-			}
-
-			if (ProcessorsToSkipThisFrame.Contains(Processor))
 			{
 				continue;
 			}
@@ -303,8 +269,6 @@ const GameInputKind FGameInputDeviceContainer::ProcessInput(IGameInput* GameInpu
 					++NumReadingsProcessed;
 				}
 
-				UsedProcessorsThisFrame.Add(Processor);
-
 				// Keep track of what input processors have sent events
 				if (bHasInputThisFrame)
 				{
