[Game Input] Make creation of the game input device async on m,odule startup for GameInput on windows. This brings it down from a worst case of 5 seconds to ~5ms on the main thread

#rnx

diff --git a/Engine/Plugins/Experimental/GameInputWindows/Source/GameInputWindows/Private/GameInputWindowsDevice.cpp b/Engine/Plugins/Experimental/GameInputWindows/Source/GameInputWindows/Private/GameInputWindowsDevice.cpp
--- a/Engine/Plugins/Experimental/GameInputWindows/Source/GameInputWindows/Private/GameInputWindowsDevice.cpp
+++ b/Engine/Plugins/Experimental/GameInputWindows/Source/GameInputWindows/Private/GameInputWindowsDevice.cpp
@@ -25,6 +25,12 @@ FGameInputWindowsInputDevice::~FGameInputWindowsInputDevice()
 #endif
 }
 
+void FGameInputWindowsInputDevice::SetGameInputAndReinitialize(IGameInput* InGameInput)
+{
+	GameInput = InGameInput;
+	Initialize();
+}
+
 GameInputKind FGameInputWindowsInputDevice::GetCurrentGameInputKindSupport() const
 {
 	GameInputKind RegisterInputKindMask = IGameInputDeviceInterface::GetCurrentGameInputKindSupport();
diff --git a/Engine/Plugins/Experimental/GameInputWindows/Source/GameInputWindows/Private/GameInputWindowsDevice.h b/Engine/Plugins/Experimental/GameInputWindows/Source/GameInputWindows/Private/GameInputWindowsDevice.h
--- a/Engine/Plugins/Experimental/GameInputWindows/Source/GameInputWindows/Private/GameInputWindowsDevice.h
+++ b/Engine/Plugins/Experimental/GameInputWindows/Source/GameInputWindows/Private/GameInputWindowsDevice.h
@@ -12,6 +12,8 @@ public:
 	explicit FGameInputWindowsInputDevice(const TSharedRef<FGenericApplicationMessageHandler>& InMessageHandler, struct IGameInput* InGameInput);
 	virtual ~FGameInputWindowsInputDevice() override;
 
+	void SetGameInputAndReinitialize(IGameInput* InGameInput);
+
 protected:
 
 	//~Begin IGameInputDeviceInterface interface	
diff --git a/Engine/Plugins/Experimental/GameInputWindows/Source/GameInputWindows/Private/GameInputWindowsModule.cpp b/Engine/Plugins/Experimental/GameInputWindows/Source/GameInputWindows/Private/GameInputWindowsModule.cpp
--- a/Engine/Plugins/Experimental/GameInputWindows/Source/GameInputWindows/Private/GameInputWindowsModule.cpp
+++ b/Engine/Plugins/Experimental/GameInputWindows/Source/GameInputWindows/Private/GameInputWindowsModule.cpp
@@ -36,38 +36,26 @@ bool FGameInputWindowsModule::IsAvailable()
 TSharedPtr<IInputDevice> FGameInputWindowsModule::CreateInputDevice(const TSharedRef<FGenericApplicationMessageHandler>& InMessageHandler)
 {
 #if GAME_INPUT_SUPPORT
-	if (IGameInput* GameInputPtr = FGameInputBaseModule::GetGameInput())
+	IGameInput* BaseGameInputPtr = FGameInputBaseModule::GetGameInput();
+	TSharedPtr<FGameInputWindowsInputDevice> DevicePtr = MakeShared<FGameInputWindowsInputDevice>(InMessageHandler, BaseGameInputPtr);
+	
+	if (BaseGameInputPtr)
 	{
-		TSharedPtr<FGameInputWindowsInputDevice> DevicePtr = MakeShared<FGameInputWindowsInputDevice>(InMessageHandler, GameInputPtr);
 		DevicePtr->Initialize();
-		return DevicePtr;
 	}
 	else
 	{
-		if (IsRunningCommandlet())
+		FGameInputBaseModule& BaseModule = FGameInputBaseModule::Get();
+		BaseModule.OnGameInputCreation.AddLambda([&DevicePtr](IGameInput* GameInputPtr)
 		{
-			// We don't care for Game Input if we are running a commandlet, like when we are cooking or running C++ automation tests
-			UE_LOG(LogGameInput, Log, TEXT("[%s] GameInput interface will not be created because IsRunningCommandlet is true."), ANSI_TO_TCHAR(__FUNCTION__));
-		}
-		else if (!FApp::HasProjectName())
-		{
-			// If there is no project name then we don't need game input either. This means we are in the project launcher
-			UE_LOG(LogGameInput, Log, TEXT("[%s] GameInput interface will not be created because there is no project name."), ANSI_TO_TCHAR(__FUNCTION__));
-		}
-		else if (FApp::IsUnattended())
-		{
-			UE_LOG(LogGameInput, Log, TEXT("[%s] GameInput interface will not be created because FApp::IsUnattended is true!"), ANSI_TO_TCHAR(__FUNCTION__));
-		}
-		else if (!FApp::CanEverRender())
-		{
-			UE_LOG(LogGameInput, Log, TEXT("[%s] GameInput interface will not be created because FApp::CanEverRender is false!"), ANSI_TO_TCHAR(__FUNCTION__));
-		}
-		else
-		{
-			// Otherwise, we want Game Input to be properly initialized
-			UE_LOG(LogGameInput, Log, TEXT("[%s] Failed to get the IGameInput object from the base module! You will not be able to process input."), ANSI_TO_TCHAR(__FUNCTION__));	
-		}
+			if (ensure(DevicePtr.IsValid()))
+			{
+				DevicePtr->SetGameInputAndReinitialize(GameInputPtr);
+			}
+		});	
 	}
+	return DevicePtr;
+	
 #else
 	UE_LOG(LogGameInput, Error, TEXT("[FGameInputWindowsModule] Failed to create a GameInput device! GAME_INPUT_SUPPORT is false! See GameInputBase.build.cs and the build log for more info."));
 #endif	// GAME_INPUT_SUPPORT
diff --git a/Engine/Plugins/Runtime/GameInput/Source/GameInputBase/Private/GameInputBaseModule.cpp b/Engine/Plugins/Runtime/GameInput/Source/GameInputBase/Private/GameInputBaseModule.cpp
--- a/Engine/Plugins/Runtime/GameInput/Source/GameInputBase/Private/GameInputBaseModule.cpp
+++ b/Engine/Plugins/Runtime/GameInput/Source/GameInputBase/Private/GameInputBaseModule.cpp
@@ -2,6 +2,7 @@
 
 #include "GameInputBaseModule.h"
 
+#include <atomic>
 #include "CoreGlobals.h"
 #include "GameInputBaseIncludes.h"
 #include "GameInputKeyTypes.h"
@@ -9,6 +10,7 @@
 #include "HAL/PlatformProcess.h"
 #include "Misc/App.h"
 #include "Modules/ModuleManager.h"
+#include "Tasks/Task.h"
 
 #if GAME_INPUT_SUPPORT
 THIRD_PARTY_INCLUDES_START
@@ -36,6 +38,10 @@ namespace UE::GameInput
 #endif	// #if GAME_INPUT_SUPPORT	
 }
 
+#if PLATFORM_WINDOWS
+	FCriticalSection FGameInputBaseModule::GameInputCreationLock;
+#endif
+
 FGameInputBaseModule& FGameInputBaseModule::Get()
 {
 	return FModuleManager::LoadModuleChecked<FGameInputBaseModule>(TEXT("GameInputBase"));
@@ -95,14 +101,58 @@ void FGameInputBaseModule::StartupModule()
 // Disable warning C4191: 'type cast' : unsafe conversion from 'PROC' to 'XXX' getting the GameInputCreate function
 #pragma warning(push)
 #pragma warning(disable:4191)
-	
-	// Create the Game Input interface
-	typedef HRESULT(WINAPI* FGameInputCreateFn)(IGameInput**);
-	FGameInputCreateFn pfnGameInputCreate = (FGameInputCreateFn)GetProcAddress((HMODULE)GameInputDLLHandle, "GameInputCreate");
 
-	UE_CLOG(!pfnGameInputCreate, LogGameInput, Warning, TEXT("[%hs] Failed to GetProcAddress (GameInputCreate). Game Input will fail to be created."), __func__);
-	
-	const HRESULT HResult = pfnGameInputCreate ? pfnGameInputCreate(&UE::GameInput::GGameInputInterface) : E_FAIL;
+	// The GameInputCreate function can be expensive on windows (~5s on startup!) so run it in an async task
+	UE::Tasks::Launch(UE_SOURCE_LOCATION, [DllHandle = GameInputDLLHandle, FinishedDelegate = MakeShared<TMulticastDelegate<void(IGameInput*)>>(OnGameInputCreation)]()
+	{
+		// Create the Game Input interface
+		typedef HRESULT(WINAPI* FGameInputCreateFn)(IGameInput**);
+		FGameInputCreateFn pfnGameInputCreate = (FGameInputCreateFn)GetProcAddress((HMODULE)DllHandle, "GameInputCreate");
+
+		UE_CLOG(!pfnGameInputCreate, LogGameInput, Warning, TEXT("[%hs] Failed to GetProcAddress (GameInputCreate). Game Input will fail to be created."), __func__);
+
+		GameInputCreationLock.Lock();
+		const HRESULT HResult = pfnGameInputCreate ? pfnGameInputCreate(&UE::GameInput::GGameInputInterface) : E_FAIL;
+		GameInputCreationLock.Unlock();
+		
+		if (SUCCEEDED(HResult))
+		{
+			UE_LOG(LogGameInput, Log, TEXT("[FGameInputBaseModule::StartupModule] Successfully created the IGameInput interface"));
+		}
+		else
+		{
+			UE_LOG(LogGameInput, Warning, TEXT("Failed to initialize GameInput: 0x%X"), HResult);	
+		}
+
+		struct FBroadcastCreationDelegateFunctor
+		{
+			TSharedPtr<TMulticastDelegate<void(IGameInput*)>> DelegatePtr = nullptr;
+			
+			explicit FBroadcastCreationDelegateFunctor(TSharedPtr<TMulticastDelegate<void(IGameInput*)>> InDelegatePtr)
+				: DelegatePtr(InDelegatePtr)
+			{}
+			
+			void operator()()
+			{
+				if (DelegatePtr)
+				{
+					DelegatePtr->Broadcast(UE::GameInput::GGameInputInterface);
+				}
+			}
+		};
+		
+		// Broadcast the creation delegate on the game thread, because IInputDevice's run only on the game thread.
+		// This module start up is running on EnginePreinit. This object is not needed until the actual engine tick
+		// when we create our IInputDevices on the platform (IInputDevice::CreateInputDevice).
+		// Since we are within an async background task, we cannot gaurantee that this will actually finish
+		// by the time the first engine tick occurs. For that reason, we need to broadcast this creation
+		// delegate so that the IInputDevice can listen for it an handle it accordingly.
+		UE::Tasks::Launch(
+			UE_SOURCE_LOCATION,
+			FBroadcastCreationDelegateFunctor(FinishedDelegate),
+			LowLevelTasks::ETaskPriority::Normal,
+			UE::Tasks::EExtendedTaskPriority::GameThreadNormalPri);
+	});
 
 #pragma warning(pop)
 
@@ -110,17 +160,13 @@ void FGameInputBaseModule::StartupModule()
 	
 	// Create the Game Input interface
 	const HRESULT HResult = GameInputCreate(&UE::GameInput::GGameInputInterface);
+	UE_CLOG(FAILED(HResult), LogGameInput, Error, TEXT("Failed to initialize GameInput: 0x%X"), HResult);
+	OnGameInputCreation.Broadcast(UE::GameInput::GGameInputInterface);
 	
 #endif	// !PLATFORM_WINDOWS
-
-	UE_CLOG(FAILED(HResult), LogGameInput, Warning, TEXT("Failed to initialize GameInput: 0x%X"), HResult);
-
-	if (SUCCEEDED(HResult))
-	{
-		UE_LOG(LogGameInput, Log, TEXT("[FGameInputBaseModule::StartupModule] Successfully created the IGameInput interface"));
-		InitializeGameInputKeys();
-	}
-
+	
+	InitializeGameInputKeys();
+	
 #else
 	UE_LOG(LogGameInput, Warning, TEXT("Failed to initalize GameInput! GAME_INPUT_SUPPORT is false!"));
 #endif	// #if GAME_INPUT_SUPPORT
@@ -129,8 +175,14 @@ void FGameInputBaseModule::StartupModule()
 void FGameInputBaseModule::ShutdownModule()
 {
 #if GAME_INPUT_SUPPORT
-	UE::GameInput::GGameInputInterface.Reset();
 
+#if PLATFORM_WINDOWS
+	FScopeLock Lock(&GameInputCreationLock);
+#endif
+
+	// GGameInputInterface must be reset before we release the DLL handle
+	UE::GameInput::GGameInputInterface.Reset();
+	
 #if PLATFORM_WINDOWS
 	if (GameInputDLLHandle)
 	{
@@ -144,6 +196,10 @@ void FGameInputBaseModule::ShutdownModule()
 #if GAME_INPUT_SUPPORT
 IGameInput* FGameInputBaseModule::GetGameInput()
 {
+#if PLATFORM_WINDOWS
+	FScopeLock ScopeLock(&GameInputCreationLock);
+#endif
+	
 	return UE::GameInput::GGameInputInterface;
 }
 #endif	// GAME_INPUT_SUPPORT
diff --git a/Engine/Plugins/Runtime/GameInput/Source/GameInputBase/Private/IGameInputDeviceInterface.cpp b/Engine/Plugins/Runtime/GameInput/Source/GameInputBase/Private/IGameInputDeviceInterface.cpp
--- a/Engine/Plugins/Runtime/GameInput/Source/GameInputBase/Private/IGameInputDeviceInterface.cpp
+++ b/Engine/Plugins/Runtime/GameInput/Source/GameInputBase/Private/IGameInputDeviceInterface.cpp
@@ -67,8 +67,13 @@ void IGameInputDeviceInterface::Tick(float DeltaTime)
 
 	FScopeLock Lock(&DeviceInfoCS);
 
-	// At this point we should always have a valid Game Input object
-	check(GameInput);
+	// At this point we normally will have a valid game input object.
+	// However, during boot it may be still getting created in the background, 
+	// so early exit in that case.
+	if (!GameInput)
+	{
+		return;
+	}
 	
 	// Process any input devices so long as we are not constrained
 	if (!bIsAppCurrentlyConstrained)
@@ -88,8 +93,13 @@ void IGameInputDeviceInterface::SendControllerEvents()
 
 	FScopeLock Lock(&DeviceInfoCS);
 
-	// At this point we should always have a valid Game Input object
-	check(GameInput);
+	// At this point we normally will have a valid game input object.
+	// However, during boot it may be still getting created in the background, 
+	// so early exit in that case.
+	if (!GameInput)
+	{
+		return;
+	}
 
 	// On the first update coming back from being constrained, we want to reset
 	// the state of inputs
diff --git a/Engine/Plugins/Runtime/GameInput/Source/GameInputBase/Public/GameInputBaseModule.h b/Engine/Plugins/Runtime/GameInput/Source/GameInputBase/Public/GameInputBaseModule.h
--- a/Engine/Plugins/Runtime/GameInput/Source/GameInputBase/Public/GameInputBaseModule.h
+++ b/Engine/Plugins/Runtime/GameInput/Source/GameInputBase/Public/GameInputBaseModule.h
@@ -2,6 +2,8 @@
 
 #pragma once
 
+#include "Delegates/Delegate.h"		// For TMulticastDelegate
+#include "HAL/CriticalSection.h"
 #include "Modules/ModuleInterface.h"
 
 #if GAME_INPUT_SUPPORT
@@ -27,6 +29,11 @@ public:
 	* Pointer to the static IGameInput that is created upon module startup.
 	*/
 	static IGameInput* GetGameInput();
+
+	/**
+	 * Delegate which is called after the creation of the IGameInput object from the game input library.
+	 */
+	TMulticastDelegate<void(IGameInput*)> OnGameInputCreation;
 #endif
 
 protected:
@@ -37,5 +44,8 @@ protected:
 	// Handle to the game input dll which is set on StartupModule.
 	// If we can't find the DLL then we will early exit and not attempt to initalize GameInput.
 	void* GameInputDLLHandle = nullptr;
+
+	static FCriticalSection GameInputCreationLock;
+	
 #endif // endif PLATFORM_WINDOWS && GAME_INPUT_SUPPORT
 };
\ No newline at end of file
