Fix SA warning in GameInput base module. HResult must be examined to appease SA

#rnx

diff --git a/Engine/Plugins/Runtime/GameInput/Source/GameInputBase/Private/GameInputBaseModule.cpp b/Engine/Plugins/Runtime/GameInput/Source/GameInputBase/Private/GameInputBaseModule.cpp
--- a/Engine/Plugins/Runtime/GameInput/Source/GameInputBase/Private/GameInputBaseModule.cpp
+++ b/Engine/Plugins/Runtime/GameInput/Source/GameInputBase/Private/GameInputBaseModule.cpp
@@ -160,7 +160,14 @@ void FGameInputBaseModule::StartupModule()
 	
 	// Create the Game Input interface
 	const HRESULT HResult = GameInputCreate(&UE::GameInput::GGameInputInterface);
-	UE_CLOG(FAILED(HResult), LogGameInput, Error, TEXT("Failed to initialize GameInput: 0x%X"), HResult);
+	if (SUCCEEDED(HResult))
+	{
+		UE_LOG(LogGameInput, Log, TEXT("[FGameInputBaseModule::StartupModule] Successfully created the IGameInput interface"));
+	}
+	else
+	{
+		UE_LOG(LogGameInput, Error, TEXT("Failed to initialize GameInput: 0x%X"), HResult);
+	}
 	OnGameInputCreation.Broadcast(UE::GameInput::GGameInputInterface);
 	
 #endif	// !PLATFORM_WINDOWS
