[GameInput]
* Fixed non unity compile error in GameInputBaseModule.cpp

diff --git a/Engine/Plugins/Runtime/GameInput/Source/GameInputBase/Private/GameInputBaseModule.cpp b/Engine/Plugins/Runtime/GameInput/Source/GameInputBase/Private/GameInputBaseModule.cpp
--- a/Engine/Plugins/Runtime/GameInput/Source/GameInputBase/Private/GameInputBaseModule.cpp
+++ b/Engine/Plugins/Runtime/GameInput/Source/GameInputBase/Private/GameInputBaseModule.cpp
@@ -7,6 +7,7 @@
 #include "GameInputBaseIncludes.h"
 #include "GameInputKeyTypes.h"
 #include "GameInputLogging.h"
+#include "HAL/CriticalSection.h"
 #include "HAL/PlatformProcess.h"
 #include "Misc/App.h"
 #include "Modules/ModuleManager.h"
