Remove ensure about bWasinitialized, which for windows can now be false if creating the device async on startup.

#rnx

diff --git a/Engine/Plugins/Runtime/GameInput/Source/GameInputBase/Private/IGameInputDeviceInterface.cpp b/Engine/Plugins/Runtime/GameInput/Source/GameInputBase/Private/IGameInputDeviceInterface.cpp
--- a/Engine/Plugins/Runtime/GameInput/Source/GameInputBase/Private/IGameInputDeviceInterface.cpp
+++ b/Engine/Plugins/Runtime/GameInput/Source/GameInputBase/Private/IGameInputDeviceInterface.cpp
@@ -89,18 +89,22 @@ void IGameInputDeviceInterface::SendControllerEvents()
 {
 #if GAME_INPUT_SUPPORT
 
-	ensureAlwaysMsgf(bWasinitialized, TEXT("A Game Input Device was not initalized. You must call IGameInputDeviceInterface::Initialize after creating a new device!"));
-
-	FScopeLock Lock(&DeviceInfoCS);
-
 	// At this point we normally will have a valid game input object.
 	// However, during boot it may be still getting created in the background, 
 	// so early exit in that case.
+	
+	if (!bWasinitialized)
+	{
+		return;
+	}
+	
 	if (!GameInput)
 	{
 		return;
 	}
 
+	FScopeLock Lock(&DeviceInfoCS);
+
 	// On the first update coming back from being constrained, we want to reset
 	// the state of inputs
 	if (bIsAppCurrentlyConstrained && !bWasAppConstrainedLastTick)
