Fix define for GameInputCreationLock

#rnx

diff --git a/Engine/Plugins/Runtime/GameInput/Source/GameInputBase/Private/GameInputBaseModule.cpp b/Engine/Plugins/Runtime/GameInput/Source/GameInputBase/Private/GameInputBaseModule.cpp
--- a/Engine/Plugins/Runtime/GameInput/Source/GameInputBase/Private/GameInputBaseModule.cpp
+++ b/Engine/Plugins/Runtime/GameInput/Source/GameInputBase/Private/GameInputBaseModule.cpp
@@ -39,7 +39,7 @@ namespace UE::GameInput
 #endif	// #if GAME_INPUT_SUPPORT	
 }
 
-#if PLATFORM_WINDOWS
+#if PLATFORM_WINDOWS && GAME_INPUT_SUPPORT
 	FCriticalSection FGameInputBaseModule::GameInputCreationLock;
 #endif
 
