[Game Input] Add a missing property initalize to zero

#rnx

diff --git a/Engine/Plugins/Runtime/GameInput/Source/GameInputBase/Public/GameInputDeviceProcessor.h b/Engine/Plugins/Runtime/GameInput/Source/GameInputBase/Public/GameInputDeviceProcessor.h
--- a/Engine/Plugins/Runtime/GameInput/Source/GameInputBase/Public/GameInputDeviceProcessor.h
+++ b/Engine/Plugins/Runtime/GameInput/Source/GameInputBase/Public/GameInputDeviceProcessor.h
@@ -327,7 +327,7 @@ protected:
 		bool bIsActive;
 		};
 	TArray<FTouchData> PreviousTouchData;
-	int32 ActiveTouchPoints;
+	int32 ActiveTouchPoints = 0;
 
 	// The max touch X and Y of the touch region to let us scale input based on the size of the touch area
 	int32 MaxTouchX = 1920;
