Fix ButtonRepeatDelay not being respected on Xbox & Switch platforms

* INI Config values for `ButtonRepeatDelay` and `InitialButtonRepeatDelay` were not being used to intialize the variables on their respective platforms.
* Updated `IGameInputDeviceProcessor`, `FNintendoInputInterface` and `FNpadController` to respect Button Repeat Delay values.

diff --git a/Engine/Plugins/Runtime/GameInput/Source/GameInputBase/Private/GameInputDeviceProcessor.cpp b/Engine/Plugins/Runtime/GameInput/Source/GameInputBase/Private/GameInputDeviceProcessor.cpp
--- a/Engine/Plugins/Runtime/GameInput/Source/GameInputBase/Private/GameInputDeviceProcessor.cpp
+++ b/Engine/Plugins/Runtime/GameInput/Source/GameInputBase/Private/GameInputDeviceProcessor.cpp
@@ -11,6 +11,7 @@
 #include "GenericPlatform/CursorUtils.h"
 #include "Framework/Application/SlateApplication.h"
 #include "HAL/IConsoleManager.h"
+#include "Misc/ConfigCacheIni.h"
 #include "GameInputDeveloperSettings.h"
 
 #if GAME_INPUT_SUPPORT
@@ -44,6 +45,15 @@ namespace UE::GameInput
 ////////////////////////////////////////////////////////
 // IGameInputDeviceProcessor
 
+IGameInputDeviceProcessor::IGameInputDeviceProcessor()
+{
+	InitialButtonRepeatDelay = UE::GameInput::InitialRepeatDelay;
+	ButtonRepeatDelay = UE::GameInput::SubsequentRepeatDelay;
+
+	GConfig->GetFloat(TEXT("/Script/Engine.InputSettings"), TEXT("InitialButtonRepeatDelay"), InitialButtonRepeatDelay, GInputIni);
+	GConfig->GetFloat(TEXT("/Script/Engine.InputSettings"), TEXT("ButtonRepeatDelay"), ButtonRepeatDelay, GInputIni);
+}
+
 const GameInputDeviceInfo* IGameInputDeviceProcessor::FGameInputEventParams::GetDeviceInfo() const
 {
 	if (Device)
@@ -185,7 +195,7 @@ void IGameInputDeviceProcessor::EvaluateButtonStates(
 				{
 					UE_LOG(LogGameInput, Verbose, TEXT("[FGameInputDevice::EvaluateButtonStates] (PlatformUserId = %d, InputDeviceId = %d) - Button '%s' Pressed"), Params.PlatformUserId.GetInternalId(), Params.InputDeviceId.GetId(), *ButtonKey.ToString());
 					Params.MessageHandler->OnControllerButtonPressed(ButtonKey, Params.PlatformUserId, Params.InputDeviceId, false);
-					RepeatTime[n] = CurrentTime + UE::GameInput::InitialRepeatDelay;
+					RepeatTime[n] = CurrentTime + InitialButtonRepeatDelay;
 				}
 				else
 				{
@@ -200,7 +210,7 @@ void IGameInputDeviceProcessor::EvaluateButtonStates(
 			{
 				FInputDeviceScope InputScope(nullptr, UE::GameInput::InputClassName, DeviceMapper.GetUserIndexForPlatformUser(Params.PlatformUserId), GetHardwareDeviceIdentifierName(Params));
 
-				RepeatTime[n] = CurrentTime + UE::GameInput::SubsequentRepeatDelay;
+				RepeatTime[n] = CurrentTime + ButtonRepeatDelay;
 				UE_LOG(LogGameInput, Verbose, TEXT("[FGameInputDevice::EvaluateButtonStates] (PlatformUserId = %d, InputDeviceId = %d) - Button '%s' Repeat Pressed"), Params.PlatformUserId.GetInternalId(), Params.InputDeviceId.GetId(), *ButtonKey.ToString());
 				Params.MessageHandler->OnControllerButtonPressed(ButtonKey, Params.PlatformUserId, Params.InputDeviceId, true);
 			}
@@ -266,7 +276,7 @@ void IGameInputDeviceProcessor::EvaluateSwitchState(
 				Params.MessageHandler->OnControllerButtonPressed(KeyName, Params.PlatformUserId, Params.InputDeviceId, true);
 			}
 
-			RepeatTimes[static_cast<int32>(CurrentPosition)] = CurrentTime + UE::GameInput::InitialRepeatDelay;
+			RepeatTimes[static_cast<int32>(CurrentPosition)] = CurrentTime + InitialButtonRepeatDelay;
 		}
 	}
 	// Otherwise, the states are the same. Check if we can repeat them
@@ -282,7 +292,7 @@ void IGameInputDeviceProcessor::EvaluateSwitchState(
 		}
 
 		// Check for repeat key
-		RepeatTimes[static_cast<int32>(CurrentPosition)] = CurrentTime + UE::GameInput::SubsequentRepeatDelay;
+		RepeatTimes[static_cast<int32>(CurrentPosition)] = CurrentTime + ButtonRepeatDelay;
 	}
 
 	// Keep track of the previous position
diff --git a/Engine/Plugins/Runtime/GameInput/Source/GameInputBase/Public/GameInputDeviceProcessor.h b/Engine/Plugins/Runtime/GameInput/Source/GameInputBase/Public/GameInputDeviceProcessor.h
--- a/Engine/Plugins/Runtime/GameInput/Source/GameInputBase/Public/GameInputDeviceProcessor.h
+++ b/Engine/Plugins/Runtime/GameInput/Source/GameInputBase/Public/GameInputDeviceProcessor.h
@@ -31,7 +31,7 @@ THIRD_PARTY_INCLUDES_END
 class GAMEINPUTBASE_API IGameInputDeviceProcessor
 {
 public:
-	IGameInputDeviceProcessor() = default;
+	IGameInputDeviceProcessor();
 
 	virtual ~IGameInputDeviceProcessor() = default;
 
@@ -148,6 +148,12 @@ protected:
 		GameInputSwitchPosition CurrentPosition,
 		GameInputSwitchPosition& PreviousPosition,
 		TArray<double>& RepeatTimes);
+
+	/** Delay before sending a repeat message after a button was first pressed */
+	float InitialButtonRepeatDelay;
+
+	/** Delay before sending a repeat message after a button has been pressed for a while */
+	float ButtonRepeatDelay;
 };
 
 /**
