[GameInputBase]
Added module dependency on SlateCore

diff --git a/Engine/Plugins/Runtime/GameInput/Source/GameInputBase/GameInputBase.build.cs b/Engine/Plugins/Runtime/GameInput/Source/GameInputBase/GameInputBase.build.cs
--- a/Engine/Plugins/Runtime/GameInput/Source/GameInputBase/GameInputBase.build.cs
+++ b/Engine/Plugins/Runtime/GameInput/Source/GameInputBase/GameInputBase.build.cs
@@ -78,6 +78,7 @@ namespace UnrealBuildTool.Rules
 				{
 					"Core",
 					"ApplicationCore",
+					"SlateCore",
 					"Slate",
 					"Engine",
 					"InputCore",
