[Gmae Input] Clean up some headers

diff --git a/Engine/Plugins/Experimental/GameInputWindows/Source/GameInputWindows/Private/GameInputWindowsModule.cpp b/Engine/Plugins/Experimental/GameInputWindows/Source/GameInputWindows/Private/GameInputWindowsModule.cpp
--- a/Engine/Plugins/Experimental/GameInputWindows/Source/GameInputWindows/Private/GameInputWindowsModule.cpp
+++ b/Engine/Plugins/Experimental/GameInputWindows/Source/GameInputWindows/Private/GameInputWindowsModule.cpp
@@ -8,14 +8,7 @@
 #include "GameInputLogging.h"
 #include "Misc/App.h"
 #include "Modules/ModuleManager.h"
-
-#if GAME_INPUT_SUPPORT
-THIRD_PARTY_INCLUDES_START
-#include "Microsoft/AllowMicrosoftPlatformTypes.h"
-#include <GameInput.h>
-#include "Microsoft/HideMicrosoftPlatformTypes.h"
-THIRD_PARTY_INCLUDES_END
-#endif	// GAME_INPUT_SUPPORT
+#include "GameInputBaseIncludes.h"
 
 namespace UE::GameInput
 {
diff --git a/Engine/Plugins/Runtime/GameInput/Source/GameInputBase/Private/GameInputBaseModule.cpp b/Engine/Plugins/Runtime/GameInput/Source/GameInputBase/Private/GameInputBaseModule.cpp
--- a/Engine/Plugins/Runtime/GameInput/Source/GameInputBase/Private/GameInputBaseModule.cpp
+++ b/Engine/Plugins/Runtime/GameInput/Source/GameInputBase/Private/GameInputBaseModule.cpp
@@ -13,15 +13,6 @@
 #include "Modules/ModuleManager.h"
 #include "Tasks/Task.h"
 
-#if GAME_INPUT_SUPPORT
-THIRD_PARTY_INCLUDES_START
-#include "Microsoft/AllowMicrosoftPlatformTypes.h"
-#include <GameInput.h>
-#include "Microsoft/HideMicrosoftPlatformTypes.h"
-THIRD_PARTY_INCLUDES_END
-#include "Microsoft/COMPointer.h"
-#endif
-
 #define LOCTEXT_NAMESPACE "GameInputBaseModule" 
 
 namespace UE::GameInput
diff --git a/Engine/Plugins/Runtime/GameInput/Source/GameInputBase/Private/GameInputDeveloperSettings.cpp b/Engine/Plugins/Runtime/GameInput/Source/GameInputBase/Private/GameInputDeveloperSettings.cpp
--- a/Engine/Plugins/Runtime/GameInput/Source/GameInputBase/Private/GameInputDeveloperSettings.cpp
+++ b/Engine/Plugins/Runtime/GameInput/Source/GameInputBase/Private/GameInputDeveloperSettings.cpp
@@ -3,14 +3,7 @@
 #include "GameInputDeveloperSettings.h"
 #include "Misc/DataDrivenPlatformInfoRegistry.h"
 #include "GameInputLogging.h"
-
-#if GAME_INPUT_SUPPORT
-THIRD_PARTY_INCLUDES_START
-#include "Microsoft/AllowMicrosoftPlatformTypes.h"
-#include <GameInput.h>
-#include "Microsoft/HideMicrosoftPlatformTypes.h"
-THIRD_PARTY_INCLUDES_END
-#endif	// GAME_INPUT_SUPPORT
+#include "GameInputBaseIncludes.h"
 
 FGameInputDeviceIdentifier::FGameInputDeviceIdentifier()
 	: VendorId(0)
diff --git a/Engine/Plugins/Runtime/GameInput/Source/GameInputBase/Public/GameInputBaseIncludes.h b/Engine/Plugins/Runtime/GameInput/Source/GameInputBase/Public/GameInputBaseIncludes.h
--- a/Engine/Plugins/Runtime/GameInput/Source/GameInputBase/Public/GameInputBaseIncludes.h
+++ b/Engine/Plugins/Runtime/GameInput/Source/GameInputBase/Public/GameInputBaseIncludes.h
@@ -4,19 +4,18 @@
 
 #if GAME_INPUT_SUPPORT
 
-// #include "Microsoft/AllowMicrosoftPlatformTypes.h"
-//
-// THIRD_PARTY_INCLUDES_START
-//
-// #include "Microsoft/COMPointer.h"
-// #include <GameInput.h>
-//
-// #if PLATFORM_WINDOWS
-// 	#include <windef.h>	// For APP_LOCAL_DEVICE_ID
-// #endif	// PLATFORM_WINDOWS
-//
-// THIRD_PARTY_INCLUDES_END
-//
-// #include "Microsoft/HideMicrosoftPlatformTypes.h"
+#include "Windows/AllowWindowsPlatformTypes.h"
+THIRD_PARTY_INCLUDES_START
+
+#if PLATFORM_WINDOWS
+	#include <windef.h>	// For APP_LOCAL_DEVICE_ID
+#endif	// PLATFORM_WINDOWS
+
+#include "Microsoft/COMPointer.h"
+
+#include "GameInput.h"
+
+THIRD_PARTY_INCLUDES_END
+#include "Windows/HideWindowsPlatformTypes.h"
 
 #endif	// GAME_INPUT_SUPPORT
\ No newline at end of file
diff --git a/Engine/Plugins/Runtime/GameInput/Source/GameInputBase/Public/GameInputDeviceProcessor.h b/Engine/Plugins/Runtime/GameInput/Source/GameInputBase/Public/GameInputDeviceProcessor.h
--- a/Engine/Plugins/Runtime/GameInput/Source/GameInputBase/Public/GameInputDeviceProcessor.h
+++ b/Engine/Plugins/Runtime/GameInput/Source/GameInputBase/Public/GameInputDeviceProcessor.h
@@ -7,15 +7,7 @@
 #include "Containers/Map.h"
 #include "Misc/CoreMiscDefines.h"	// For FInputDeviceId/FPlatformUserId
 #include "GenericPlatform/GenericApplicationMessageHandler.h"
-
-#if GAME_INPUT_SUPPORT
-THIRD_PARTY_INCLUDES_START
-#include "Microsoft/AllowMicrosoftPlatformTypes.h"
-#include <GameInput.h>
-#include "Microsoft/HideMicrosoftPlatformTypes.h"
-THIRD_PARTY_INCLUDES_END
-#include "Microsoft/COMPointer.h"
-#endif	// GAME_INPUT_SUPPORT
+#include "GameInputBaseIncludes.h"
 
 #if GAME_INPUT_SUPPORT
 
diff --git a/Engine/Plugins/Runtime/GameInput/Source/GameInputBase/Public/GameInputUtils.h b/Engine/Plugins/Runtime/GameInput/Source/GameInputBase/Public/GameInputUtils.h
--- a/Engine/Plugins/Runtime/GameInput/Source/GameInputBase/Public/GameInputUtils.h
+++ b/Engine/Plugins/Runtime/GameInput/Source/GameInputBase/Public/GameInputUtils.h
@@ -7,12 +7,7 @@
 #include "Containers/UnrealString.h"
 #include "GenericPlatform/GenericPlatformInputDeviceMapper.h"
 #include "GenericPlatform/GenericApplicationMessageHandler.h"
-
-THIRD_PARTY_INCLUDES_START
-#include "Microsoft/AllowMicrosoftPlatformTypes.h"
-#include <GameInput.h>
-#include "Microsoft/HideMicrosoftPlatformTypes.h"
-THIRD_PARTY_INCLUDES_END
+#include "GameInputBaseIncludes.h"
 
 namespace UE::GameInput
 {
diff --git a/Engine/Plugins/Runtime/GameInput/Source/GameInputBase/Public/IGameInputDeviceInterface.h b/Engine/Plugins/Runtime/GameInput/Source/GameInputBase/Public/IGameInputDeviceInterface.h
--- a/Engine/Plugins/Runtime/GameInput/Source/GameInputBase/Public/IGameInputDeviceInterface.h
+++ b/Engine/Plugins/Runtime/GameInput/Source/GameInputBase/Public/IGameInputDeviceInterface.h
@@ -4,16 +4,7 @@
 
 #include "IInputDevice.h"
 #include "Misc/CoreMiscDefines.h"	// For FInputDeviceId/FPlatformUserId
-
-#if GAME_INPUT_SUPPORT
-THIRD_PARTY_INCLUDES_START
-#include "Microsoft/AllowMicrosoftPlatformTypes.h"
-#include <GameInput.h>
-#include "Microsoft/HideMicrosoftPlatformTypes.h"
-THIRD_PARTY_INCLUDES_END
-#include "Microsoft/COMPointer.h"
-#endif	// GAME_INPUT_SUPPORT
-
+#include "GameInputBaseIncludes.h"
 #include "GameInputDeviceContainer.h"
 
 /**
