[Game Input] Refactor so that the base implementation allows for different versions of Game Input headers to be used. This improves compatibilty between the NuGet and GDK versions, which have varying levels of support for some features.

- Remove some forward declarations in favor of header includes to handle the newer "Game Input" namespacing from the NuGet versions 1.0+ (the GameInput::v1 namespace)
- Wrapper for getting the device input which will handle the different versions correctly
- Add pre-processor defines for features that vary between the GDK/Nuget versions (raw and touch support) based on the current version number.

diff --git a/Engine/Plugins/Experimental/GameInputWindows/Source/GameInputWindows/Private/GameInputWindowsDevice.cpp b/Engine/Plugins/Experimental/GameInputWindows/Source/GameInputWindows/Private/GameInputWindowsDevice.cpp
--- a/Engine/Plugins/Experimental/GameInputWindows/Source/GameInputWindows/Private/GameInputWindowsDevice.cpp
+++ b/Engine/Plugins/Experimental/GameInputWindows/Source/GameInputWindows/Private/GameInputWindowsDevice.cpp
@@ -105,7 +105,7 @@ void FGameInputWindowsInputDevice::HandleDeviceDisconnected(IGameInputDevice* De
 void FGameInputWindowsInputDevice::HandleDeviceConnected(IGameInputDevice* Device, uint64 Timestamp)
 {
 	// get device information
-	const GameInputDeviceInfo* Info = Device->GetDeviceInfo();
+	const GameInputDeviceInfo* Info = UE::GameInput::GetDeviceInfo(Device);
 	UE_LOG(LogGameInput, Log, TEXT("Game Input Device Connected: %s  of kind: %s"), *UE::GameInput::LexToString(Device), *UE::GameInput::LexToString(Info->supportedInput));
 	
 	FGameInputDeviceContainer* Data = GetOrCreateDeviceData(Device);
diff --git a/Engine/Plugins/Runtime/GameInput/Source/GameInputBase/Private/GameInputBaseModule.cpp b/Engine/Plugins/Runtime/GameInput/Source/GameInputBase/Private/GameInputBaseModule.cpp
--- a/Engine/Plugins/Runtime/GameInput/Source/GameInputBase/Private/GameInputBaseModule.cpp
+++ b/Engine/Plugins/Runtime/GameInput/Source/GameInputBase/Private/GameInputBaseModule.cpp
@@ -4,7 +4,6 @@
 
 #include <atomic>
 #include "CoreGlobals.h"
-#include "GameInputBaseIncludes.h"
 #include "GameInputKeyTypes.h"
 #include "GameInputLogging.h"
 #include "HAL/CriticalSection.h"
diff --git a/Engine/Plugins/Runtime/GameInput/Source/GameInputBase/Private/GameInputDeviceContainer.cpp b/Engine/Plugins/Runtime/GameInput/Source/GameInputBase/Private/GameInputDeviceContainer.cpp
--- a/Engine/Plugins/Runtime/GameInput/Source/GameInputBase/Private/GameInputDeviceContainer.cpp
+++ b/Engine/Plugins/Runtime/GameInput/Source/GameInputBase/Private/GameInputDeviceContainer.cpp
@@ -30,7 +30,7 @@ FGameInputDeviceContainer::FGameInputDeviceContainer(
 	}
 
 	// Initalize the App Local ID, which should never change on this container
-	if (const GameInputDeviceInfo* Info = InDevice->GetDeviceInfo())
+	if (const GameInputDeviceInfo* Info = UE::GameInput::GetDeviceInfo(Device))
 	{
 		LocalDeviceId = Info->deviceId;
 	}
@@ -49,7 +49,7 @@ void FGameInputDeviceContainer::InitalizeDeviceProcessors()
 	}
 
 	// Based on the GameInputKind of this device, create any processors that are supported
-	const GameInputDeviceInfo* Info = Device->GetDeviceInfo();
+	const GameInputDeviceInfo* Info = UE::GameInput::GetDeviceInfo(Device);
 
 	const GameInputKind DeviceKind = Info->supportedInput;
 
@@ -82,7 +82,7 @@ void FGameInputDeviceContainer::InitalizeDeviceProcessors_Impl()
 	}
 
 	// Based on the GameInputKind of this device, create any processors that are supported
-	const GameInputDeviceInfo* Info = Device->GetDeviceInfo();
+	const GameInputDeviceInfo* Info = UE::GameInput::GetDeviceInfo(Device);
 
 	const GameInputKind DeviceKind = Info->supportedInput;
 
@@ -103,10 +103,12 @@ void FGameInputDeviceContainer::InitalizeDeviceProcessors_Impl()
 	//	Processors.Emplace(MakeShared<FGameInputMouseDeviceProcessor>(FSlateApplication::Get().GetPlatformCursor()));
 	//}
 
+#if UE_GAMEINPUT_SUPPORTS_TOUCH
 	if (DeviceKind & AllowedGameInputKinds & GameInputKindTouch)
 	{
 		Processors.Emplace(MakeShared<FGameInputTouchDeviceProcessor>());
 	}
+#endif	// UE_GAMEINPUT_SUPPORTS_TOUCH
 
 	// Only allow for external controller devices have have specified configs if that is enabled
 	const UGameInputDeveloperSettings* Settings = GetDefault<UGameInputDeveloperSettings>();
@@ -134,6 +136,7 @@ void FGameInputDeviceContainer::InitalizeDeviceProcessors_Impl()
 		}
 	}
 
+#if UE_GAMEINPUT_SUPPORTS_RAW
 	// Raw Input
 	if (DeviceKind & AllowedGameInputKinds & GameInputKindRawDeviceReport)
 	{
@@ -155,6 +158,7 @@ void FGameInputDeviceContainer::InitalizeDeviceProcessors_Impl()
 				*UE::GameInput::LexToString(Device));
 		}
 	}
+#endif	// UE_GAMEINPUT_SUPPORTS_RAW
 
 	// Racing wheels
 	if (DeviceKind & AllowedGameInputKinds & GameInputKindRacingWheel)
@@ -337,7 +341,7 @@ void FGameInputDeviceContainer::SetGameInputDevice(IGameInputDevice* InDevice)
 		// 
 		// In this case, we would get here if the IGameInputDevice pointer was set to null upon disconnection, 
 		// and then the device was re-connected.
-		if (const GameInputDeviceInfo* Info = InDevice->GetDeviceInfo())
+		if (const GameInputDeviceInfo* Info = UE::GameInput::GetDeviceInfo(Device))
 		{
 			const bool bAppIdsAreTheSame = (FMemory::Memcmp(&Info->deviceId, &LocalDeviceId, sizeof(LocalDeviceId)) == 0);
 			ensure(bAppIdsAreTheSame);
diff --git a/Engine/Plugins/Runtime/GameInput/Source/GameInputBase/Private/GameInputDeviceProcessor.cpp b/Engine/Plugins/Runtime/GameInput/Source/GameInputBase/Private/GameInputDeviceProcessor.cpp
--- a/Engine/Plugins/Runtime/GameInput/Source/GameInputBase/Private/GameInputDeviceProcessor.cpp
+++ b/Engine/Plugins/Runtime/GameInput/Source/GameInputBase/Private/GameInputDeviceProcessor.cpp
@@ -56,11 +56,7 @@ IGameInputDeviceProcessor::IGameInputDeviceProcessor()
 
 const GameInputDeviceInfo* IGameInputDeviceProcessor::FGameInputEventParams::GetDeviceInfo() const
 {
-	if (Device)
-	{
-		return Device->GetDeviceInfo();
-	}
-	return nullptr;
+	return UE::GameInput::GetDeviceInfo(Device);
 }
 
 bool IGameInputDeviceProcessor::PostProcessInput(const FGameInputEventParams& Params)
@@ -1307,6 +1303,8 @@ GameInputKind FGameInputMouseDeviceProcessor::GetSupportedReadingKind() const
 ////////////////////////////////////////////////////////
 // FGameInputTouchDeviceProcessor
 
+#if UE_GAMEINPUT_SUPPORTS_TOUCH
+
 bool FGameInputTouchDeviceProcessor::ProcessInput(const FGameInputEventParams& Params)
 {
 	// read the new touch events
@@ -1423,6 +1421,10 @@ GameInputKind FGameInputTouchDeviceProcessor::GetSupportedReadingKind() const
 	return GameInputKindTouch;
 }
 
+#endif	// #if UE_GAMEINPUT_SUPPORTS_TOUCH
+
+#if UE_GAMEINPUT_SUPPORTS_RAW
+
 constexpr float FGameInputRawDeviceProcessor::RawValueToFloatTrigger(const uint8 RawValue) const
 {
 	// Maps the uint8 value of 0-255 to a float between 0.0 and +1.0, like a gamepad trigger.
@@ -1817,6 +1819,8 @@ GameInputKind FGameInputRawDeviceProcessor::GetSupportedReadingKind() const
 	return GameInputKindRawDeviceReport;
 }
 
+#endif	// UE_GAMEINPUT_SUPPORTS_RAW
+
 ///////////////////////////////////////////////////////////////////////////////////////
 // FGameInputRacingWheelProcessor
 
diff --git a/Engine/Plugins/Runtime/GameInput/Source/GameInputBase/Private/GameInputUtils.cpp b/Engine/Plugins/Runtime/GameInput/Source/GameInputBase/Private/GameInputUtils.cpp
--- a/Engine/Plugins/Runtime/GameInput/Source/GameInputBase/Private/GameInputUtils.cpp
+++ b/Engine/Plugins/Runtime/GameInput/Source/GameInputBase/Private/GameInputUtils.cpp
@@ -12,13 +12,20 @@ namespace UE::GameInput
 
 		if (Device)
 		{
-			const GameInputDeviceInfo* Info = Device->GetDeviceInfo();
+			const GameInputDeviceInfo* Info = UE::GameInput::GetDeviceInfo(Device);
 			const bool bIsVirtualDevice = Info->deviceFamily == GameInputFamilyVirtual;
 
+#if GAMEINPUT_API_VERSION <= 0
 			const FString DeviceDisplayName =
 				Info->displayName ?
 				FString::Printf(TEXT("%.*hs"), Info->displayName->sizeInBytes, Info->displayName->data) :
 				TEXT("Unknown");
+#else
+			const FString DeviceDisplayName =
+				Info->displayName ?
+				FString::Printf(TEXT("%s"), ANSI_TO_TCHAR(Info->displayName)) :
+				TEXT("Unknown");
+#endif
 			
 			Result = FString::Printf(TEXT("%p   DeviceName: %s %s ProdId: %04x (%u) VendId: %04x (%u)"),
 				Device, 
@@ -45,6 +52,7 @@ namespace UE::GameInput
 		FString Result = TEXT("");
 #define DEVICE_STATUS(StatusFlag,DisplayName) if( DeviceStatus & StatusFlag ) Result += (FString(DisplayName) + TEXT("|"));
 		DEVICE_STATUS(GameInputDeviceConnected, TEXT("Connected"));
+#if UE_GAMEINPUT_SUPPORTS_DEVICE_STATUS
 		DEVICE_STATUS(GameInputDeviceInputEnabled, TEXT("InputEnabled"));
 		DEVICE_STATUS(GameInputDeviceOutputEnabled, TEXT("OutputEnabled"));
 		DEVICE_STATUS(GameInputDeviceRawIoEnabled, TEXT("RawIoEnabled"));
@@ -53,6 +61,7 @@ namespace UE::GameInput
 		DEVICE_STATUS(GameInputDeviceSynchronized, TEXT("Synchronized"));
 		DEVICE_STATUS(GameInputDeviceWireless, TEXT("Wireless"));
 		DEVICE_STATUS(GameInputDeviceUserIdle, TEXT("UserIdle"));
+#endif
 #undef DEVICE_STATUS
 
 		Result.RemoveFromEnd(TEXT("|"));
@@ -68,15 +77,19 @@ namespace UE::GameInput
 
 		FString Result = TEXT("");
 #define INPUT_KIND_STRING(StatusFlag, DisplayName) if( InputKind & StatusFlag ) Result += (FString(DisplayName) + TEXT("|"));
+#if UE_GAMEINPUT_SUPPORTS_RAW
 		INPUT_KIND_STRING(GameInputKindRawDeviceReport, TEXT("RawDeviceReport"));
+#endif
 		INPUT_KIND_STRING(GameInputKindControllerAxis, TEXT("ControllerAxis"));
 		INPUT_KIND_STRING(GameInputKindControllerButton, TEXT("ControllerButton"));
 		INPUT_KIND_STRING(GameInputKindControllerSwitch, TEXT("ControllerSwitch"));
 		INPUT_KIND_STRING(GameInputKindController, TEXT("Controller"));
 		INPUT_KIND_STRING(GameInputKindKeyboard, TEXT("Keyboard"));
 		INPUT_KIND_STRING(GameInputKindMouse, TEXT("Mouse"));
+#if UE_GAMEINPUT_SUPPORTS_TOUCH
 		INPUT_KIND_STRING(GameInputKindTouch, TEXT("Touch"));
 		INPUT_KIND_STRING(GameInputKindMotion, TEXT("Motion"));
+#endif
 		INPUT_KIND_STRING(GameInputKindArcadeStick, TEXT("ArcadeStick"));
 		INPUT_KIND_STRING(GameInputKindFlightStick, TEXT("FlightStick"));
 		INPUT_KIND_STRING(GameInputKindGamepad, TEXT("Gamepad"));
@@ -178,6 +191,25 @@ namespace UE::GameInput
 	{
 		return GetSwitchButtonMap().Find(ButtonMask);
 	}
+
+	const GameInputDeviceInfo* GetDeviceInfo(IGameInputDevice* Device)
+	{
+		if (!Device)
+		{
+			return nullptr;
+		}
+
+		const GameInputDeviceInfo* Info = nullptr;
+		
+#if GAMEINPUT_API_VERSION >= 1
+		Device->GetDeviceInfo(&Info);
+#else
+		Info = Device->GetDeviceInfo();
+#endif
+
+		return Info;
+	}
+	
 }
 
 #endif	// #if GAME_INPUT_SUPPORT
diff --git a/Engine/Plugins/Runtime/GameInput/Source/GameInputBase/Private/IGameInputDeviceInterface.cpp b/Engine/Plugins/Runtime/GameInput/Source/GameInputBase/Private/IGameInputDeviceInterface.cpp
--- a/Engine/Plugins/Runtime/GameInput/Source/GameInputBase/Private/IGameInputDeviceInterface.cpp
+++ b/Engine/Plugins/Runtime/GameInput/Source/GameInputBase/Private/IGameInputDeviceInterface.cpp
@@ -26,6 +26,7 @@ IGameInputDeviceInterface::IGameInputDeviceInterface(const TSharedRef<FGenericAp
 #if GAME_INPUT_SUPPORT
 	: MessageHandler(InMessageHandler)
 	, GameInput(InGameInput)
+	, ConnectionChangeCallbackToken{}
 	, CurrentlyConnectedDeviceTypes(GameInputKindUnknown)
 	, bWasinitialized(false)
 	, bIsAppCurrentlyConstrained(false)
@@ -41,11 +42,20 @@ IGameInputDeviceInterface::~IGameInputDeviceInterface()
 	{
 		if (GameInput)
 		{
+#if GAMEINPUT_API_VERSION >= 1
+			GameInput->UnregisterCallback(ConnectionChangeCallbackToken);
+#else
 			GameInput->UnregisterCallback(ConnectionChangeCallbackToken, UINT64_MAX);
+#endif
+			
 			GameInput.Reset();
 		}
 
+#if GAMEINPUT_API_VERSION >= 1
+		ConnectionChangeCallbackToken = {};
+#else
 		ConnectionChangeCallbackToken = GAMEINPUT_INVALID_CALLBACK_TOKEN_VALUE;
+#endif
 	}
 
 	FCoreDelegates::ApplicationWillDeactivateDelegate.RemoveAll(this);
@@ -285,10 +295,12 @@ GameInputKind IGameInputDeviceInterface::GetCurrentGameInputKindSupport() const
 		RegisterInputKindMask |= (GameInputKindController | GameInputKindControllerAxis | GameInputKindControllerButton | GameInputKindControllerSwitch);
 	}
 
+#if UE_GAMEINPUT_SUPPORTS_RAW
 	if (PlatformSettings->bProcessRawInput)
 	{
 		RegisterInputKindMask |= GameInputKindRawDeviceReport;
 	}
+#endif
 
 	if (PlatformSettings->bProcessGamepad)
 	{
@@ -333,8 +345,13 @@ bool IGameInputDeviceInterface::BindDeviceStatusCallbacks()
 	// First, check if we are already bound to something. If so, unbind it
 	if (ConnectionChangeCallbackToken)
 	{
+#if GAMEINPUT_API_VERSION >= 1
+		GameInput->UnregisterCallback(ConnectionChangeCallbackToken);
+		ConnectionChangeCallbackToken = {};
+#else
 		GameInput->UnregisterCallback(ConnectionChangeCallbackToken, UINT64_MAX);
 		ConnectionChangeCallbackToken = GAMEINPUT_INVALID_CALLBACK_TOKEN_VALUE;
+#endif
 	}
 
 	// Hook up device callbacks 
@@ -357,7 +374,11 @@ bool IGameInputDeviceInterface::BindDeviceStatusCallbacks()
 	FCoreDelegates::ApplicationHasReactivatedDelegate.AddRaw(this, &IGameInputDeviceInterface::OnAppUnconstrained);
 
 	// We are successful if the callback token is valid
+#if GAMEINPUT_API_VERSION >= 1
+	return true;
+#else
 	return ConnectionChangeCallbackToken != GAMEINPUT_INVALID_CALLBACK_TOKEN_VALUE;
+#endif
 }
 
 void IGameInputDeviceInterface::OnDeviceConnectionChanged(GameInputCallbackToken CallbackToken, IGameInputDevice* Device, uint64 Timestamp, GameInputDeviceStatus CurrentStatus, GameInputDeviceStatus PreviousStatus)
@@ -435,7 +456,7 @@ FGameInputDeviceContainer* IGameInputDeviceInterface::GetDeviceData(IGameInputDe
 
 	// Check if we have seen this device's APP_LOCAL_DEVICE_ID before. The IGameInputDevice* could have been invalidated if the device 
 	// is being re-connected, but it will have the APP_LOCAL_DEVICE_ID would be the same.
-	const GameInputDeviceInfo* Info = InDevice->GetDeviceInfo();
+	const GameInputDeviceInfo* Info = UE::GameInput::GetDeviceInfo(InDevice);
 
 	FGameInputDeviceContainer* RetVal = nullptr;
 
@@ -493,7 +514,7 @@ void IGameInputDeviceInterface::EnumerateCurrentlyConnectedDeviceTypes()
 		{
 			if (IGameInputDevice* Device = KnownDevice->GetGameInputDevice())
 			{
-				if (const GameInputDeviceInfo* Info = Device->GetDeviceInfo())
+				if (const GameInputDeviceInfo* Info = UE::GameInput::GetDeviceInfo(Device))				
 				{
 					CurrentlyConnectedDeviceTypes |= Info->supportedInput;
 				}
diff --git a/Engine/Plugins/Runtime/GameInput/Source/GameInputBase/Public/GameInputBaseIncludes.h b/Engine/Plugins/Runtime/GameInput/Source/GameInputBase/Public/GameInputBaseIncludes.h
--- a/Engine/Plugins/Runtime/GameInput/Source/GameInputBase/Public/GameInputBaseIncludes.h
+++ b/Engine/Plugins/Runtime/GameInput/Source/GameInputBase/Public/GameInputBaseIncludes.h
@@ -15,6 +15,35 @@ THIRD_PARTY_INCLUDES_START
 
 #include "GameInput.h"
 
+#ifndef GAMEINPUT_API_VERSION
+	#define GAMEINPUT_API_VERSION 0
+#endif
+
+#if GAMEINPUT_API_VERSION == 0
+
+#elif GAMEINPUT_API_VERSION == 1
+
+	using namespace GameInput::v1;
+
+#endif
+
+// Depending on the Game Input API version it will support different types of input processing.
+// At the moment Game Input v1 does not support touch or raw input device processing, so we will
+// compile those device processors out. "v0", which is the GDK version, does support these methods
+// and will have them compiled in.
+
+#if GAMEINPUT_API_VERSION <= 0
+
+	#define UE_GAMEINPUT_SUPPORTS_RAW 1
+	#define UE_GAMEINPUT_SUPPORTS_TOUCH 1
+	#define UE_GAMEINPUT_SUPPORTS_DEVICE_STATUS 1
+#else
+	#define UE_GAMEINPUT_SUPPORTS_RAW 0
+	#define UE_GAMEINPUT_SUPPORTS_TOUCH 0
+	#define UE_GAMEINPUT_SUPPORTS_DEVICE_STATUS 0
+#endif
+
+
 THIRD_PARTY_INCLUDES_END
 #include "Windows/HideWindowsPlatformTypes.h"
 
diff --git a/Engine/Plugins/Runtime/GameInput/Source/GameInputBase/Public/GameInputBaseModule.h b/Engine/Plugins/Runtime/GameInput/Source/GameInputBase/Public/GameInputBaseModule.h
--- a/Engine/Plugins/Runtime/GameInput/Source/GameInputBase/Public/GameInputBaseModule.h
+++ b/Engine/Plugins/Runtime/GameInput/Source/GameInputBase/Public/GameInputBaseModule.h
@@ -5,12 +5,10 @@
 #include "Delegates/Delegate.h"		// For TMulticastDelegate
 #include "HAL/CriticalSection.h"
 #include "Modules/ModuleInterface.h"
+#include "GameInputBaseIncludes.h"
 
 #define UE_API GAMEINPUTBASE_API
 
-#if GAME_INPUT_SUPPORT
-struct IGameInput;
-#endif
 
 class FGameInputBaseModule : public IModuleInterface
 {
diff --git a/Engine/Plugins/Runtime/GameInput/Source/GameInputBase/Public/GameInputDeveloperSettings.h b/Engine/Plugins/Runtime/GameInput/Source/GameInputBase/Public/GameInputDeveloperSettings.h
--- a/Engine/Plugins/Runtime/GameInput/Source/GameInputBase/Public/GameInputDeveloperSettings.h
+++ b/Engine/Plugins/Runtime/GameInput/Source/GameInputBase/Public/GameInputDeveloperSettings.h
@@ -4,6 +4,7 @@
 
 #include "Engine/DeveloperSettings.h"
 #include "Engine/PlatformSettings.h"
+#include "GameInputBaseIncludes.h"
 
 #include "GameInputDeveloperSettings.generated.h"
 
@@ -278,8 +279,6 @@ struct FGameInputDeviceConfiguration
 	TMap<int32, FGameInputRawDeviceReportData> RawReportMappingData;
 };
 
-struct GameInputDeviceInfo;
-
 /**
  * Settings for Game Input that you may want to set per-platform.
  */
diff --git a/Engine/Plugins/Runtime/GameInput/Source/GameInputBase/Public/GameInputDeviceProcessor.h b/Engine/Plugins/Runtime/GameInput/Source/GameInputBase/Public/GameInputDeviceProcessor.h
--- a/Engine/Plugins/Runtime/GameInput/Source/GameInputBase/Public/GameInputDeviceProcessor.h
+++ b/Engine/Plugins/Runtime/GameInput/Source/GameInputBase/Public/GameInputDeviceProcessor.h
@@ -297,6 +297,8 @@ protected:
 };
 
 
+#if UE_GAMEINPUT_SUPPORTS_TOUCH
+
 /**
 * Processor for touch inputs
 */
@@ -331,6 +333,9 @@ protected:
 	int32 MaxTouchX = 1920;
 	int32 MaxTouchY = 1080;
 };
+#endif	// #if UE_GAMEINPUT_SUPPORTS_TOUCH
+
+#if UE_GAMEINPUT_SUPPORTS_RAW
 
 struct FGameInputDeviceConfiguration;
 struct FGameInputRawDeviceReportData;
@@ -444,6 +449,8 @@ protected:
 	int32 NumReadingsProcessedThisFrame = 0;
 };
 
+#endif	// UE_GAMEINPUT_SUPPORTS_RAW
+
 /**
 * Processor for racing wheel types (GameInputKindRacingWheel)
 * 
diff --git a/Engine/Plugins/Runtime/GameInput/Source/GameInputBase/Public/GameInputUtils.h b/Engine/Plugins/Runtime/GameInput/Source/GameInputBase/Public/GameInputUtils.h
--- a/Engine/Plugins/Runtime/GameInput/Source/GameInputBase/Public/GameInputUtils.h
+++ b/Engine/Plugins/Runtime/GameInput/Source/GameInputBase/Public/GameInputUtils.h
@@ -42,6 +42,8 @@ namespace UE::GameInput
 	*/
 	GAMEINPUTBASE_API const TArray<FGamepadKeyNames::Type>* SwitchPositionToUnrealName(const GameInputSwitchPosition ButtonMask);
 
+	GAMEINPUTBASE_API const GameInputDeviceInfo* GetDeviceInfo(IGameInputDevice* Device);
+	
 	// Button flags assigned by GameInput
 	enum { GamingInputButtonMask = 0x3FFF };
 
