[Game Input] Make sure that the Application correctly has it's key modifier state updated before we call any "OnKeyDown" or "OnKeyUp" functions when using a keyboard.

[RN] Bug fix: Ensure that the "SlateApp.GetModifierKeys" function will return the correct values when called on platforms using GameInput to process keyboards

diff --git a/Engine/Plugins/Runtime/GameInput/Source/GameInputBase/Public/GameInputDeviceProcessor.h b/Engine/Plugins/Runtime/GameInput/Source/GameInputBase/Public/GameInputDeviceProcessor.h
--- a/Engine/Plugins/Runtime/GameInput/Source/GameInputBase/Public/GameInputDeviceProcessor.h
+++ b/Engine/Plugins/Runtime/GameInput/Source/GameInputBase/Public/GameInputDeviceProcessor.h
@@ -254,7 +254,7 @@ protected:
 	GAMEINPUTBASE_API virtual void ClearState(const FGameInputEventParams& Params) override;
 	GAMEINPUTBASE_API virtual GameInputKind GetSupportedReadingKind() const override;
 
-	GAMEINPUTBASE_API void UpdateUnifiedKeyboardState(const FGameInputEventParams& Params, TSet<uint8>& CurrentPressedKeys);
+	GAMEINPUTBASE_API virtual void UpdateUnifiedKeyboardState(const FGameInputEventParams& Params, TSet<uint8>& CurrentPressedKeys);
 	
 	inline bool IsKeyPressed(uint8 VirtualKeyCode) const { return LastPressedKeys.Contains(VirtualKeyCode); }
 
