[GameInput] Add cvar for enabling touch input correction in the context of FGameInputTouchDeviceProcessor

-Add a new cvar that allows correcting the touch pressure value in case the reading received is not set/incorrect. Also add a cvar for controlling the default value in case this happens.

diff --git a/Engine/Plugins/Runtime/GameInput/Source/GameInputBase/Private/GameInputDeviceProcessor.cpp b/Engine/Plugins/Runtime/GameInput/Source/GameInputBase/Private/GameInputDeviceProcessor.cpp
--- a/Engine/Plugins/Runtime/GameInput/Source/GameInputBase/Private/GameInputDeviceProcessor.cpp
+++ b/Engine/Plugins/Runtime/GameInput/Source/GameInputBase/Private/GameInputDeviceProcessor.cpp
@@ -40,7 +40,22 @@ namespace UE::GameInput
 
 		return RacingWheelButtonMap;
 	}
-};
+}; // namespace UE::GameInput
+
+namespace GameInputDevicePreProcessorCVars
+{
+	bool bEnableTouchPressureCorrection = false;
+	FAutoConsoleVariableRef CVarEnableTouchPressureCorrection(
+		TEXT("GameInputDeviceProcessor.EnableTouchPressureCorrection"),
+		bEnableTouchPressureCorrection,
+		TEXT("This enables overriding touch input pressure in case the readings received by the touch input device preprocessor is considered too small."));
+
+	float TouchPressureCorrectionValue = 1.f;
+	FAutoConsoleVariableRef CVarTouchPressureCorrectionValue(
+		TEXT("GameInputDeviceProcessor.TouchPressureCorrectionValue"),
+		TouchPressureCorrectionValue,
+		TEXT("Touch pressure value to be used in case the touch input device preprocessor receives a reading that is considered to be too small"));
+}; // namespace GameInputDevicePreProcessorCVars
 
 ////////////////////////////////////////////////////////
 // IGameInputDeviceProcessor
@@ -1365,9 +1380,13 @@ bool FGameInputTouchDeviceProcessor::ProcessInput(const FGameInputEventParams& P
 		// need to scale by resolution
 		FVector2D NewPosition = FVector2D(InputTouch.positionX * (float)MaxTouchX, InputTouch.positionY * (float)MaxTouchY).RoundToVector();
 
+		// reading may fail to have a proper nonzero touch pressure, so we correct it
+		const bool bShouldCorrectTouchPressure = GameInputDevicePreProcessorCVars::bEnableTouchPressureCorrection && FMath::Abs(InputTouch.pressure) < GameInputDevicePreProcessorCVars::TouchPressureCorrectionValue;
+		const float NewTouchPressure = bShouldCorrectTouchPressure ? GameInputDevicePreProcessorCVars::TouchPressureCorrectionValue : InputTouch.pressure;
+
 		if (TouchState.TouchId == INDEX_NONE)
 		{
-			Params.MessageHandler->OnTouchStarted(nullptr, NewPosition, InputTouch.pressure, Index, Params.PlatformUserId, Params.InputDeviceId);
+			Params.MessageHandler->OnTouchStarted(nullptr, NewPosition, NewTouchPressure, Index, Params.PlatformUserId, Params.InputDeviceId);
 			TouchState.TouchId = InputTouch.touchId;
 			ActiveTouchPoints++;
 		}
@@ -1375,21 +1394,21 @@ bool FGameInputTouchDeviceProcessor::ProcessInput(const FGameInputEventParams& P
 		{
 			if (!TouchState.bHasMoved)
 			{
-				Params.MessageHandler->OnTouchFirstMove(NewPosition, InputTouch.pressure, Index, Params.PlatformUserId, Params.InputDeviceId);
+				Params.MessageHandler->OnTouchFirstMove(NewPosition, NewTouchPressure, Index, Params.PlatformUserId, Params.InputDeviceId);
 				TouchState.bHasMoved = true;
 			}
 			else
 			{
-				Params.MessageHandler->OnTouchMoved(NewPosition, InputTouch.pressure, Index, Params.PlatformUserId, Params.InputDeviceId);
+				Params.MessageHandler->OnTouchMoved(NewPosition, NewTouchPressure, Index, Params.PlatformUserId, Params.InputDeviceId);
 			}
 		}
 
-		if (TouchState.Pressure != InputTouch.pressure)
+		if (TouchState.Pressure != NewTouchPressure)
 		{
-			Params.MessageHandler->OnTouchForceChanged(NewPosition, InputTouch.pressure, Index, Params.PlatformUserId, Params.InputDeviceId);
+			Params.MessageHandler->OnTouchForceChanged(NewPosition, NewTouchPressure, Index, Params.PlatformUserId, Params.InputDeviceId);
 		}
 
-		TouchState.Pressure = InputTouch.pressure;
+		TouchState.Pressure = NewTouchPressure;
 		TouchState.Position = NewPosition;
 		TouchState.bIsActive = true;
 	}
