[Game Input Windows]

Add a UE_GAMEINPUT_REQUIRES_DYNAMIC_CREATION pre-processor define, which is true on platforms which do not have the GameInput.lib static library available to them, but are still on the Windows platform.

If UE_GAMEINPUT_REQUIRES_DYNAMIC_CREATION is false, like on the "normal" Win64 target, then we can simply call the CreateGameInput function and use the static library. If it is false, then we still need to dynamically load the GameInput.DLL file and call it via a function pointer.

[RN] When possible, favor the static GameInput.lib functions instead of requiring a dynamic DLL load.

diff --git a/Engine/Plugins/Runtime/GameInput/Source/GameInputBase/GameInputBase.build.cs b/Engine/Plugins/Runtime/GameInput/Source/GameInputBase/GameInputBase.build.cs
--- a/Engine/Plugins/Runtime/GameInput/Source/GameInputBase/GameInputBase.build.cs
+++ b/Engine/Plugins/Runtime/GameInput/Source/GameInputBase/GameInputBase.build.cs
@@ -3,7 +3,10 @@
 namespace UnrealBuildTool.Rules
 {
 	public class GameInputBase : ModuleRules
-	{		
+	{
+		// If true, then we will use the dynamic GameInput.dll file instead of a statically linked GameInput library.
+		protected virtual bool bRequiresDynamicCreation { get { return false; } }
+
 		/// <summary>
 		/// True if this platform has support for the Game Input library.
 		/// 
@@ -81,6 +84,8 @@ namespace UnrealBuildTool.Rules
 			// Define this as 0 in the base module to avoid compilation errors when building
 			// without any Game Input support. It is up to the platform-specific submodules to define 
 			PublicDefinitions.Add("GAME_INPUT_SUPPORT=" + (bHasGameInputSupport ? "1" : "0"));
+			
+			PublicDefinitions.Add("UE_GAMEINPUT_REQUIRES_DYNAMIC_CREATION=" + (bRequiresDynamicCreation ? "1" : "0"));
 
 			// Give platforms extensions a chance to add any required dependencies that may be necessary to compile game input
 			if (bHasGameInputSupport)
diff --git a/Engine/Plugins/Runtime/GameInput/Source/GameInputBase/Private/GameInputBaseModule.cpp b/Engine/Plugins/Runtime/GameInput/Source/GameInputBase/Private/GameInputBaseModule.cpp
--- a/Engine/Plugins/Runtime/GameInput/Source/GameInputBase/Private/GameInputBaseModule.cpp
+++ b/Engine/Plugins/Runtime/GameInput/Source/GameInputBase/Private/GameInputBaseModule.cpp
@@ -21,10 +21,10 @@ namespace UE::GameInput
 	// This provides access to reading the input stream, device callbacks, and more. 
 	static TComPtr<IGameInput> GGameInputInterface;
 
-#if PLATFORM_WINDOWS
+#if UE_GAMEINPUT_REQUIRES_DYNAMIC_CREATION
 	// The name of the Game Input DLL file which would be needed when running on windows
 	static const FString GameInputDLLPath = TEXT("GameInput.dll");
-#endif	// #if PLATFORM_WINDOWS
+#endif	// #if UE_GAMEINPUT_REQUIRES_DYNAMIC_CREATION
 
 #endif	// #if GAME_INPUT_SUPPORT	
 }
@@ -50,34 +50,36 @@ void FGameInputBaseModule::StartupModule()
 	// We don't care for Game Input if we are running a commandlet, like when we are cooking.
 	if (IsRunningCommandlet())
 	{
-		UE_LOG(LogGameInput, Log, TEXT("GameInputBase module is exiting because IsRunningCommandlet is true. GameInput will not be initalized."));
+		UE_LOG(LogGameInput, Log, TEXT("GameInputBase module is exiting because IsRunningCommandlet is true. GameInput will not be initialized."));
 		return;
 	}
 
 	// If there is no project name then we don't need game input either. This means we are in the project launcher
 	if (!FApp::HasProjectName())
 	{
-		UE_LOG(LogGameInput, Log, TEXT("GameInputBase module is exiting because there is no project name. GameInput will not be initalized."));
+		UE_LOG(LogGameInput, Log, TEXT("GameInputBase module is exiting because there is no project name. GameInput will not be initialized."));
 		return;
 	}
 
 	// Unattended app can't receive any user input, so there is no need to try and create the GameInput interface.
 	if (FApp::IsUnattended() && !FApp::AllowUnattendedInput())
 	{
-		UE_LOG(LogGameInput, Log, TEXT("GameInputBase module is exiting because it is unattended (FApp::IsUnattended is true) and thus cannot recieve user input. GameInput will not be initalized."));
+		UE_LOG(LogGameInput, Log, TEXT("GameInputBase module is exiting because it is unattended (FApp::IsUnattended is true) and thus cannot receive user input. GameInput will not be initialized."));
 		return;
 	}
 
 	// Doesn't make sense to have headless apps create game input
 	if (!FApp::CanEverRender())
 	{
-		UE_LOG(LogGameInput, Log, TEXT("GameInputBase module is exiting because it cannot render anything (FApp::CanEverRender is false). GameInput will not be initalized."));
+		UE_LOG(LogGameInput, Log, TEXT("GameInputBase module is exiting because it cannot render anything (FApp::CanEverRender is false). GameInput will not be initialized."));
 		return;
 	}
 
 #if GAME_INPUT_SUPPORT
 
 #if PLATFORM_WINDOWS
+
+#if UE_GAMEINPUT_REQUIRES_DYNAMIC_CREATION
 	// Check to see if the GameInput.dll exists...
 	// Search for the GameInput dll on desktop platforms. If for some reason it doesn't exist, then we shouldn't
 	// attempt to call any functions from it. The only known case for this is when running a client on a server OS
@@ -85,25 +87,39 @@ void FGameInputBaseModule::StartupModule()
 	GameInputDLLHandle = FPlatformProcess::GetDllHandle(*UE::GameInput::GameInputDLLPath);
 	if (GameInputDLLHandle == nullptr)
 	{
-		UE_LOG(LogGameInput, Warning, TEXT("[%hs] module is exiting because '%s' cannot be found. GameInput will not be initalized. Is it installed correctly?"), __func__, *UE::GameInput::GameInputDLLPath);
+		UE_LOG(LogGameInput, Warning, TEXT("[%hs] module is exiting because '%s' cannot be found. GameInput will not be initialized. Is it installed correctly?"), __func__, *UE::GameInput::GameInputDLLPath);
 		return;
 	}
+#endif	// UE_GAMEINPUT_REQUIRES_DYNAMIC_CREATION
 
 // Disable warning C4191: 'type cast' : unsafe conversion from 'PROC' to 'XXX' getting the GameInputCreate function
 #pragma warning(push)
 #pragma warning(disable:4191)
 
 	// The GameInputCreate function can be expensive on windows (~5s on startup!) so run it in an async task
+#if UE_GAMEINPUT_REQUIRES_DYNAMIC_CREATION
 	UE::Tasks::Launch(UE_SOURCE_LOCATION, [DllHandle = GameInputDLLHandle, FinishedDelegate = MakeShared<TMulticastDelegate<void(IGameInput*)>>(OnGameInputCreation)]()
+#else
+	UE::Tasks::Launch(UE_SOURCE_LOCATION, [FinishedDelegate = MakeShared<TMulticastDelegate<void(IGameInput*)>>(OnGameInputCreation)]()
+#endif	// #if UE_GAMEINPUT_REQUIRES_DYNAMIC_CREATION	
 	{
+
+#if UE_GAMEINPUT_REQUIRES_DYNAMIC_CREATION
 		// Create the Game Input interface
 		typedef HRESULT(WINAPI* FGameInputCreateFn)(IGameInput**);
 		FGameInputCreateFn pfnGameInputCreate = (FGameInputCreateFn)GetProcAddress((HMODULE)DllHandle, "GameInputCreate");
 
 		UE_CLOG(!pfnGameInputCreate, LogGameInput, Warning, TEXT("[%hs] Failed to GetProcAddress (GameInputCreate). Game Input will fail to be created."), __func__);
+#endif	// UE_GAMEINPUT_REQUIRES_DYNAMIC_CREATION
 
 		GameInputCreationLock.Lock();
+
+#if UE_GAMEINPUT_REQUIRES_DYNAMIC_CREATION
 		const HRESULT HResult = pfnGameInputCreate ? pfnGameInputCreate(&UE::GameInput::GGameInputInterface) : E_FAIL;
+#else
+		const HRESULT HResult = GameInputCreate(&UE::GameInput::GGameInputInterface);
+#endif	// #if UE_GAMEINPUT_REQUIRES_DYNAMIC_CREATION
+
 		GameInputCreationLock.Unlock();
 		
 		if (SUCCEEDED(HResult))
@@ -166,7 +182,7 @@ void FGameInputBaseModule::StartupModule()
 	InitializeGameInputKeys();
 	
 #else
-	UE_LOG(LogGameInput, Warning, TEXT("Failed to initalize GameInput! GAME_INPUT_SUPPORT is false!"));
+	UE_LOG(LogGameInput, Warning, TEXT("Failed to initialize GameInput! GAME_INPUT_SUPPORT is false!"));
 #endif	// #if GAME_INPUT_SUPPORT
 }
 
@@ -181,12 +197,12 @@ void FGameInputBaseModule::ShutdownModule()
 	// GGameInputInterface must be reset before we release the DLL handle
 	UE::GameInput::GGameInputInterface.Reset();
 	
-#if PLATFORM_WINDOWS
+#if UE_GAMEINPUT_REQUIRES_DYNAMIC_CREATION
 	if (GameInputDLLHandle)
 	{
 		FPlatformProcess::FreeDllHandle(GameInputDLLHandle);
 	}
-#endif // #if PLATFORM_WINDOWS
+#endif // #if UE_GAMEINPUT_REQUIRES_DYNAMIC_CREATION
 
 #endif // #if GAME_INPUT_SUPPORT
 }
diff --git a/Engine/Plugins/Runtime/GameInput/Source/GameInputBase/Public/GameInputBaseModule.h b/Engine/Plugins/Runtime/GameInput/Source/GameInputBase/Public/GameInputBaseModule.h
--- a/Engine/Plugins/Runtime/GameInput/Source/GameInputBase/Public/GameInputBaseModule.h
+++ b/Engine/Plugins/Runtime/GameInput/Source/GameInputBase/Public/GameInputBaseModule.h
@@ -41,9 +41,12 @@ protected:
 	UE_API void InitializeGameInputKeys();
 
 #if PLATFORM_WINDOWS && GAME_INPUT_SUPPORT
+
+#if UE_GAMEINPUT_REQUIRES_DYNAMIC_CREATION
 	// Handle to the game input dll which is set on StartupModule.
-	// If we can't find the DLL then we will early exit and not attempt to initalize GameInput.
+	// If we can't find the DLL then we will early exit and not attempt to initialize GameInput.
 	void* GameInputDLLHandle = nullptr;
+#endif	// UE_GAMEINPUT_REQUIRES_DYNAMIC_CREATION
 
 	static UE_API FCriticalSection GameInputCreationLock;
 	
