[Game Input] Also print out the custom display name of a device upon connection, as well as the version of game input currently in use for easier to read logging

diff --git a/Engine/Plugins/Runtime/GameInput/Source/GameInputBase/Private/GameInputBaseModule.cpp b/Engine/Plugins/Runtime/GameInput/Source/GameInputBase/Private/GameInputBaseModule.cpp
--- a/Engine/Plugins/Runtime/GameInput/Source/GameInputBase/Private/GameInputBaseModule.cpp
+++ b/Engine/Plugins/Runtime/GameInput/Source/GameInputBase/Private/GameInputBaseModule.cpp
@@ -77,6 +77,8 @@ void FGameInputBaseModule::StartupModule()
 
 #if GAME_INPUT_SUPPORT
 
+	UE_LOG(LogGameInput, Log, TEXT("Using GameInput Version: %d"), GAMEINPUT_API_VERSION);
+	
 #if PLATFORM_WINDOWS
 	// Check to see if the GameInput.dll exists...
 	// Search for the GameInput dll on desktop platforms. If for some reason it doesn't exist, then we shouldn't
diff --git a/Engine/Plugins/Runtime/GameInput/Source/GameInputBase/Private/GameInputUtils.cpp b/Engine/Plugins/Runtime/GameInput/Source/GameInputBase/Private/GameInputUtils.cpp
--- a/Engine/Plugins/Runtime/GameInput/Source/GameInputBase/Private/GameInputUtils.cpp
+++ b/Engine/Plugins/Runtime/GameInput/Source/GameInputBase/Private/GameInputUtils.cpp
@@ -2,6 +2,8 @@
 
 #include "GameInputUtils.h"
 
+#include "GameInputDeveloperSettings.h"
+
 #if GAME_INPUT_SUPPORT
 
 namespace UE::GameInput
@@ -26,11 +28,21 @@ namespace UE::GameInput
 				FString::Printf(TEXT("%s"), ANSI_TO_TCHAR(Info->displayName)) :
 				TEXT("Unknown");
 #endif
+
+			// If we have a custom device setup for this, print that name too
+			const UGameInputDeveloperSettings* Settings = GetDefault<UGameInputDeveloperSettings>();
+
+			const FGameInputDeviceConfiguration* CustomDeviceConfig = Settings->FindDeviceConfiguration(Info);
+			const FString CustomHardwareId =
+				CustomDeviceConfig ?
+				FString::Printf(TEXT("(%s)"), *CustomDeviceConfig->OverriddenHardwareDeviceId) :
+				TEXT("");
 			
-			Result = FString::Printf(TEXT("%p   DeviceName: %s %s ProdId: %04x (%u) VendId: %04x (%u)"),
+			Result = FString::Printf(TEXT("%p   DeviceName: %s %s %s ProdId: %04x (%u) VendId: %04x (%u)"),
 				Device, 
 				*DeviceDisplayName,
 				bIsVirtualDevice ? TEXT("(virtual)") : TEXT(""),
+				*CustomHardwareId,
 				Info->productId, (uint32)Info->productId,
 				Info->vendorId, (uint32)Info->vendorId);
 		}
