[Game Input Windows]

Add a UE_GAMEINPUT_REQUIRES_DYNAMIC_CREATION pre-processor define, which is true on platforms which do not have the GameInput.lib static library available to them, but are still on the Windows platform.

If UE_GAMEINPUT_REQUIRES_DYNAMIC_CREATION is false, like on the "normal" Win64 target, then we can simply call the CreateGameInput function and use the static library. If it is false, then we still need to dynamically load the GameInput.DLL file and call it via a function pointer.

[RN] When possible, favor the static GameInput.lib functions instead of requiring a dynamic DLL load.

diff --git a/Engine/Plugins/Runtime/GameInput/Source/GameInputBase/GameInputBase.build.cs b/Engine/Plugins/Runtime/GameInput/Source/GameInputBase/GameInputBase.build.cs
--- a/Engine/Plugins/Runtime/GameInput/Source/GameInputBase/GameInputBase.build.cs
+++ b/Engine/Plugins/Runtime/GameInput/Source/GameInputBase/GameInputBase.build.cs
@@ -1,9 +1,14 @@
 // Copyright Epic Games, Inc. All Rights Reserved.
 
+using System.Linq;
+
 namespace UnrealBuildTool.Rules
 {
 	public class GameInputBase : ModuleRules
-	{		
+	{
+		// If true, then we will use the dynamic GameInput.dll file instead of a statically linked GameInput library.
+		protected virtual bool bRequiresDynamicCreation { get { return false; } }
+
 		/// <summary>
 		/// True if this platform has support for the Game Input library.
 		/// 
@@ -78,9 +83,15 @@ namespace UnrealBuildTool.Rules
 			);
 
 			bool bHasGameInputSupport = HasGameInputSupport(Target);
+			
 			// Define this as 0 in the base module to avoid compilation errors when building
 			// without any Game Input support. It is up to the platform-specific submodules to define 
 			PublicDefinitions.Add("GAME_INPUT_SUPPORT=" + (bHasGameInputSupport ? "1" : "0"));
+			
+			// Some targets may want to disable this plugin entirely
+			bool bShouldForceDynamicCreation = Target.GlobalDefinitions.Contains("UE_GAME_INPUT_FORCE_DYNAMIC_CREATION=1");
+			
+			PublicDefinitions.Add("UE_GAMEINPUT_REQUIRES_DYNAMIC_CREATION=" + (bShouldForceDynamicCreation || bRequiresDynamicCreation ? "1" : "0"));
 
 			// Give platforms extensions a chance to add any required dependencies that may be necessary to compile game input
 			if (bHasGameInputSupport)
diff --git a/Engine/Plugins/Runtime/GameInput/Source/GameInputBase/Private/GameInputBaseModule.cpp b/Engine/Plugins/Runtime/GameInput/Source/GameInputBase/Private/GameInputBaseModule.cpp
--- a/Engine/Plugins/Runtime/GameInput/Source/GameInputBase/Private/GameInputBaseModule.cpp
+++ b/Engine/Plugins/Runtime/GameInput/Source/GameInputBase/Private/GameInputBaseModule.cpp
@@ -50,28 +50,28 @@ void FGameInputBaseModule::StartupModule()
 	// We don't care for Game Input if we are running a commandlet, like when we are cooking.
 	if (IsRunningCommandlet())
 	{
-		UE_LOG(LogGameInput, Log, TEXT("GameInputBase module is exiting because IsRunningCommandlet is true. GameInput will not be initalized."));
+		UE_LOG(LogGameInput, Log, TEXT("GameInputBase module is exiting because IsRunningCommandlet is true. GameInput will not be initialized."));
 		return;
 	}
 
 	// If there is no project name then we don't need game input either. This means we are in the project launcher
 	if (!FApp::HasProjectName())
 	{
-		UE_LOG(LogGameInput, Log, TEXT("GameInputBase module is exiting because there is no project name. GameInput will not be initalized."));
+		UE_LOG(LogGameInput, Log, TEXT("GameInputBase module is exiting because there is no project name. GameInput will not be initialized."));
 		return;
 	}
 
 	// Unattended app can't receive any user input, so there is no need to try and create the GameInput interface.
 	if (FApp::IsUnattended() && !FApp::AllowUnattendedInput())
 	{
-		UE_LOG(LogGameInput, Log, TEXT("GameInputBase module is exiting because it is unattended (FApp::IsUnattended is true) and thus cannot recieve user input. GameInput will not be initalized."));
+		UE_LOG(LogGameInput, Log, TEXT("GameInputBase module is exiting because it is unattended (FApp::IsUnattended is true) and thus cannot receive user input. GameInput will not be initialized."));
 		return;
 	}
 
 	// Doesn't make sense to have headless apps create game input
 	if (!FApp::CanEverRender())
 	{
-		UE_LOG(LogGameInput, Log, TEXT("GameInputBase module is exiting because it cannot render anything (FApp::CanEverRender is false). GameInput will not be initalized."));
+		UE_LOG(LogGameInput, Log, TEXT("GameInputBase module is exiting because it cannot render anything (FApp::CanEverRender is false). GameInput will not be initialized."));
 		return;
 	}
 
@@ -87,7 +87,7 @@ void FGameInputBaseModule::StartupModule()
 	GameInputDLLHandle = FPlatformProcess::GetDllHandle(*UE::GameInput::GameInputDLLPath);
 	if (GameInputDLLHandle == nullptr)
 	{
-		UE_LOG(LogGameInput, Warning, TEXT("[%hs] module is exiting because '%s' cannot be found. GameInput will not be initalized. Is it installed correctly?"), __func__, *UE::GameInput::GameInputDLLPath);
+		UE_LOG(LogGameInput, Warning, TEXT("[%hs] module is exiting because '%s' cannot be found. GameInput will not be initialized. Is it installed correctly?"), __func__, *UE::GameInput::GameInputDLLPath);
 		return;
 	}
 
@@ -97,15 +97,31 @@ void FGameInputBaseModule::StartupModule()
 
 	// The GameInputCreate function can be expensive on windows (~5s on startup!) so run it in an async task
 	UE::Tasks::Launch(UE_SOURCE_LOCATION, [DllHandle = GameInputDLLHandle, FinishedDelegate = MakeShared<TMulticastDelegate<void(IGameInput*)>>(OnGameInputCreation)]()
-	{
-		// Create the Game Input interface
+	{		
+		// Create the Game Input interface via the dynamic DLL function call.
 		typedef HRESULT(WINAPI* FGameInputCreateFn)(IGameInput**);
 		FGameInputCreateFn pfnGameInputCreate = (FGameInputCreateFn)GetProcAddress((HMODULE)DllHandle, "GameInputCreate");
 
-		UE_CLOG(!pfnGameInputCreate, LogGameInput, Warning, TEXT("[%hs] Failed to GetProcAddress (GameInputCreate). Game Input will fail to be created."), __func__);
-
+		// We need to check that we have a valid GameInputCreate function available on the system, otherwise it will crash the application. The only known
+		// case for this is on Windows Server targets and some virtualized containers. 
+		if (!pfnGameInputCreate)
+		{
+			UE_LOG(LogGameInput, Warning, TEXT("[%hs] Failed to GetProcAddress (GameInputCreate). Game Input will fail to be created."), __func__);
+			return;
+		}
+		
 		GameInputCreationLock.Lock();
+
+#ifndef UE_GAME_INPUT_FORCE_DYNAMIC_CREATION
+	#define UE_GAME_INPUT_FORCE_DYNAMIC_CREATION 0
+#endif
+			
+#if UE_GAMEINPUT_REQUIRES_DYNAMIC_CREATION || UE_GAME_INPUT_FORCE_DYNAMIC_CREATION	
 		const HRESULT HResult = pfnGameInputCreate ? pfnGameInputCreate(&UE::GameInput::GGameInputInterface) : E_FAIL;
+#else
+		const HRESULT HResult = GameInputCreate(&UE::GameInput::GGameInputInterface);
+#endif
+			
 		GameInputCreationLock.Unlock();
 		
 		if (SUCCEEDED(HResult))
@@ -168,7 +184,7 @@ void FGameInputBaseModule::StartupModule()
 	InitializeGameInputKeys();
 	
 #else
-	UE_LOG(LogGameInput, Warning, TEXT("Failed to initalize GameInput! GAME_INPUT_SUPPORT is false!"));
+	UE_LOG(LogGameInput, Warning, TEXT("Failed to initialize GameInput! GAME_INPUT_SUPPORT is false!"));
 #endif	// #if GAME_INPUT_SUPPORT
 }
 
