[Game Input] Add first pass of gyro/acceleramotor support for supported devices.

This is currently compiled out, as it only works with v2.0+ and we are currently on v0 for back compat reasons, but we will be upgrading shortly.

#wip
#rnx

diff --git a/Engine/Plugins/Runtime/GameInput/Source/GameInputBase/Private/GameInputBaseModule.cpp b/Engine/Plugins/Runtime/GameInput/Source/GameInputBase/Private/GameInputBaseModule.cpp
--- a/Engine/Plugins/Runtime/GameInput/Source/GameInputBase/Private/GameInputBaseModule.cpp
+++ b/Engine/Plugins/Runtime/GameInput/Source/GameInputBase/Private/GameInputBaseModule.cpp
@@ -277,6 +277,13 @@ void FGameInputBaseModule::InitializeGameInputKeys()
 	EKeys::AddKey(FKeyDetails(FGameInputKeys::ArcadeStick_Action6, LOCTEXT("GameInput_ArcadeStick_Action6", "Game Input Arcade Stick Action 6"), FKeyDetails::GamepadKey, MenuCategory));
 	EKeys::AddKey(FKeyDetails(FGameInputKeys::ArcadeStick_Special1, LOCTEXT("GameInput_ArcadeStick_Special1", "Game Input Arcade Stick Special 1"), FKeyDetails::GamepadKey, MenuCategory));
 	EKeys::AddKey(FKeyDetails(FGameInputKeys::ArcadeStick_Special2, LOCTEXT("GameInput_ArcadeStick_Special2", "Game Input Arcade Stick Special 2"), FKeyDetails::GamepadKey, MenuCategory));
+	
+	//
+	// Sensor Types
+	//
+	
+	// Analog types
+	EKeys::AddKey(FKeyDetails(FGameInputKeys::Sensor_HeadingInDegreesFromMagneticNorth, LOCTEXT("GameInput_Sensor_HeadingInDegreesFromMagneticNorth", "Game Input Heading From Magnetic North (In Degrees)"), FKeyDetails::Axis1D, MenuCategory));
 
 #endif	// GAME_INPUT_SUPPORT
 }
diff --git a/Engine/Plugins/Runtime/GameInput/Source/GameInputBase/Private/GameInputDeviceContainer.cpp b/Engine/Plugins/Runtime/GameInput/Source/GameInputBase/Private/GameInputDeviceContainer.cpp
--- a/Engine/Plugins/Runtime/GameInput/Source/GameInputBase/Private/GameInputDeviceContainer.cpp
+++ b/Engine/Plugins/Runtime/GameInput/Source/GameInputBase/Private/GameInputDeviceContainer.cpp
@@ -177,6 +177,14 @@ void FGameInputDeviceContainer::InitalizeDeviceProcessors_Impl()
 	{
 		Processors.Emplace(MakeShared<FGameInputFlightStickProcessor>());
 	}
+
+	// If this gamepad device has any "Sensor" support, then add that as well.
+#if UE_GAMEINPUT_SUPPORTS_SENSORS
+	if (DeviceKind & AllowedGameInputKinds & GameInputKindSensors)
+	{
+		Processors.Emplace(MakeShared<FGameInputGamepadSensorDeviceProcessor>());
+	}
+#endif
 }
 
 const GameInputKind FGameInputDeviceContainer::ProcessInput(IGameInput* GameInput, const GameInputKind CurrentSupportedKind, const GameInputKind ProcessedKindsForPlatformUserThisFrame)
diff --git a/Engine/Plugins/Runtime/GameInput/Source/GameInputBase/Private/GameInputDeviceProcessor.cpp b/Engine/Plugins/Runtime/GameInput/Source/GameInputBase/Private/GameInputDeviceProcessor.cpp
--- a/Engine/Plugins/Runtime/GameInput/Source/GameInputBase/Private/GameInputDeviceProcessor.cpp
+++ b/Engine/Plugins/Runtime/GameInput/Source/GameInputBase/Private/GameInputDeviceProcessor.cpp
@@ -597,6 +597,135 @@ GameInputKind FGameInputGamepadDeviceProcessor::GetSupportedReadingKind() const
 	return GameInputKindGamepad;
 }
 
+//////////////////////////////////////////////////////////
+// FGameInputGamepadSensorDeviceProcessor
+
+#if UE_GAMEINPUT_SUPPORTS_SENSORS
+
+bool FGameInputGamepadSensorDeviceProcessor::ProcessInput(const FGameInputEventParams& Params)
+{
+	// Do nothing in this function because all sensor data is float/analog values, so we only want to
+	// check it once per-frame and not (potentially) more. 
+	return false;
+}
+
+bool FGameInputGamepadSensorDeviceProcessor::PostProcessInput(const FGameInputEventParams& Params)
+{
+	bool bRes = false;
+
+	// Can't do anything for an invalid platform user
+	if (!Params.PlatformUserId.IsValid() || !Params.PreviousReading)
+	{
+		return bRes;
+	}
+
+	// Get the current sensor state of this frame.
+	GameInputSensorsState SensorState = {};
+	if (!Params.PreviousReading->GetSensorsState(&SensorState))
+	{
+		return bRes;
+	}
+
+	// Actually process the sensor analog data 
+	bRes |= ProcessGamepadSensorState(Params, SensorState);
+
+	// Track our previous sensor state for comparison on the next frame
+	PreviousState = SensorState;
+
+	return bRes;
+}
+
+void FGameInputGamepadSensorDeviceProcessor::ClearState(const FGameInputEventParams& Params)
+{
+	// Can't do anything for an invalid platform user
+	if (!Params.PlatformUserId.IsValid())
+	{
+		return;
+	}
+
+	// We can simply process a sensor state where everything is zero
+	GameInputSensorsState ZeroValueState = {};
+	FMemory::Memset(ZeroValueState, 0);
+
+	ProcessGamepadSensorState(Params, ZeroValueState);
+
+	// Make sure to update the previous state here
+	PreviousState = ZeroValueState;
+}
+
+bool FGameInputGamepadSensorDeviceProcessor::ProcessGamepadSensorState(const FGameInputEventParams& Params, GameInputSensorsState& SensorState)
+{
+	const GameInputDeviceInfo* Info = Params.GetDeviceInfo();
+	if (!Info || !Info->sensorsInfo)
+	{
+		return false;
+	}
+	
+	const GameInputSensorsKind DetectedSensorKind = Info->sensorsInfo->supportedSensors;
+	
+	// This device does not support any sensors. Don't bother trying to get a reading from it.
+	if (DetectedSensorKind == GameInputSensorsNone)
+	{
+		return false;
+	}
+
+	bool bHasSensorReading = false;
+	
+	static constexpr GameInputSensorsKind RequiredMotionDataKinds = (GameInputSensorsGyrometer | GameInputSensorsAccelerometer);
+	
+	// Orientation, accel, and gyro data 
+	if (DetectedSensorKind & RequiredMotionDataKinds)
+	{
+		const FVector AngularVelocity = { SensorState.angularVelocityInRadPerSecX, SensorState.angularVelocityInRadPerSecY, SensorState.angularVelocityInRadPerSecZ };
+		const FVector Accel = { SensorState.accelerationInGX, SensorState.accelerationInGY, SensorState.accelerationInGZ };
+		
+		const FQuat Orientation(SensorState.orientationW, SensorState.orientationX, SensorState.orientationY, SensorState.orientationZ);
+		const FRotator Rotation = Orientation.Rotator();
+
+		// TODO_BH: Unify motion space settings. This is using the "1" setting from other platforms 
+		
+		Params.MessageHandler->OnMotionDetected(
+				FVector(Rotation.Pitch, Rotation.Yaw, Rotation.Roll),
+				FVector(AngularVelocity.Z, -AngularVelocity.X, -AngularVelocity.Y),
+				FVector::ZeroVector,
+				FVector(-Accel.Z, Accel.X, Accel.Y),
+				Params.PlatformUserId,
+				Params.InputDeviceId
+			);
+
+		UE_LOG(LogGameInput, VeryVerbose, TEXT("[%hs] OnMotionDetected PlatUser: %d DeviceId %d  MotionSpace: %d"),
+			__func__,
+			Params.PlatformUserId.GetInternalId(),
+			Params.InputDeviceId.GetId(),
+			1);
+
+		bHasSensorReading = true;
+	}
+
+	if (DetectedSensorKind & GameInputSensorsCompass)
+	{
+		// TODO_BH: We have SensorState.headingAccuracy, so perhaps we can have a different deadzone setting for different accuracies.
+		OnControllerAnalog(
+			Params,
+			FGameInputKeys::Sensor_HeadingInDegreesFromMagneticNorth.GetFName(),
+			SensorState.headingInDegreesFromMagneticNorth,
+			PreviousState.headingInDegreesFromMagneticNorth,
+			UE::GameInput::SensorMagneticNorthDeadzone);
+
+		UE_LOG(LogGameInput, VeryVerbose, TEXT("[ProcessGamepadSensorState] (Device %s) Degrees From Magnetic North: %.2f (Accuracy: %s) "),
+			*UE::GameInput::LexToString(Params.Device), SensorState.headingInDegreesFromMagneticNorth, *UE::GameInput::LexToString(SensorState.headingAccuracy));
+	}
+	
+	return bHasSensorReading;
+}
+
+GameInputKind FGameInputGamepadSensorDeviceProcessor::GetSupportedReadingKind() const
+{
+	return GameInputKindSensors;
+}
+
+#endif	// #if UE_GAMEINPUT_SUPPORTS_SENSORS
+
 ////////////////////////////////////////////////////////
 // FGameInputControllerDeviceProcessor
 
diff --git a/Engine/Plugins/Runtime/GameInput/Source/GameInputBase/Private/GameInputKeyTypes.cpp b/Engine/Plugins/Runtime/GameInput/Source/GameInputBase/Private/GameInputKeyTypes.cpp
--- a/Engine/Plugins/Runtime/GameInput/Source/GameInputBase/Private/GameInputKeyTypes.cpp
+++ b/Engine/Plugins/Runtime/GameInput/Source/GameInputBase/Private/GameInputKeyTypes.cpp
@@ -50,4 +50,11 @@ const FKey FGameInputKeys::ArcadeStick_Action4("GameInput_ArcadeStick_Action4");
 const FKey FGameInputKeys::ArcadeStick_Action5("GameInput_ArcadeStick_Action5");
 const FKey FGameInputKeys::ArcadeStick_Action6("GameInput_ArcadeStick_Action6");
 const FKey FGameInputKeys::ArcadeStick_Special1("GameInput_ArcadeStick_Special1");
-const FKey FGameInputKeys::ArcadeStick_Special2("GameInput_ArcadeStick_Special2");
\ No newline at end of file
+const FKey FGameInputKeys::ArcadeStick_Special2("GameInput_ArcadeStick_Special2");
+
+//
+// Sensor Types
+//
+
+// Analog Types
+const FKey FGameInputKeys::Sensor_HeadingInDegreesFromMagneticNorth("GameInput_Sensor_HeadingInDegreesFromMagneticNorth");
\ No newline at end of file
diff --git a/Engine/Plugins/Runtime/GameInput/Source/GameInputBase/Private/GameInputUtils.cpp b/Engine/Plugins/Runtime/GameInput/Source/GameInputBase/Private/GameInputUtils.cpp
--- a/Engine/Plugins/Runtime/GameInput/Source/GameInputBase/Private/GameInputUtils.cpp
+++ b/Engine/Plugins/Runtime/GameInput/Source/GameInputBase/Private/GameInputUtils.cpp
@@ -98,6 +98,9 @@ namespace UE::GameInput
 		INPUT_KIND_STRING(GameInputKindController, TEXT("Controller"));
 		INPUT_KIND_STRING(GameInputKindKeyboard, TEXT("Keyboard"));
 		INPUT_KIND_STRING(GameInputKindMouse, TEXT("Mouse"));
+#if UE_GAMEINPUT_SUPPORTS_SENSORS
+		INPUT_KIND_STRING(GameInputKindSensors, TEXT("Sensors"));
+#endif
 #if UE_GAMEINPUT_SUPPORTS_TOUCH
 		INPUT_KIND_STRING(GameInputKindTouch, TEXT("Touch"));
 		INPUT_KIND_STRING(GameInputKindMotion, TEXT("Motion"));
@@ -135,6 +138,45 @@ namespace UE::GameInput
 	{
 		return BytesToHex(DeviceId.value, APP_LOCAL_DEVICE_ID_SIZE);
 	}
+
+#if UE_GAMEINPUT_SUPPORTS_SENSORS
+	FString LexToString(GameInputSensorsKind InputKind)
+	{
+		if (InputKind == GameInputSensorsNone)
+		{
+			return TEXT("None");
+		}
+
+		FString Result = TEXT("");
+#define INPUT_KIND_STRING(StatusFlag, DisplayName) if( InputKind & StatusFlag ) Result += (FString(DisplayName) + TEXT("|"));
+		INPUT_KIND_STRING(GameInputSensorsAccelerometer, TEXT("Accelerometer"));
+		INPUT_KIND_STRING(GameInputSensorsGyrometer, TEXT("Gyrometer"));
+		INPUT_KIND_STRING(GameInputSensorsCompass, TEXT("Compass"));
+		INPUT_KIND_STRING(GameInputSensorsOrientation, TEXT("Orientation"));
+#undef INPUT_KIND_STRING
+
+		Result.RemoveFromEnd(TEXT("|"));
+		return Result;
+	}
+
+	FString LexToString(GameInputSensorAccuracy InputKind)
+	{
+		if (InputKind == GameInputSensorAccuracyUnknown)
+		{
+			return TEXT("Unknown");
+		}
+
+		FString Result = TEXT("");
+#define INPUT_KIND_STRING(StatusFlag, DisplayName) if( InputKind & StatusFlag ) Result += (FString(DisplayName) + TEXT("|"));
+		INPUT_KIND_STRING(GameInputSensorAccuracyUnreliable, TEXT("Unreliable"));
+		INPUT_KIND_STRING(GameInputSensorAccuracyApproximate, TEXT("Approximate"));
+		INPUT_KIND_STRING(GameInputSensorAccuracyHigh, TEXT("High"));
+#undef INPUT_KIND_STRING
+
+		Result.RemoveFromEnd(TEXT("|"));
+		return Result;
+	}
+#endif	// #if UE_GAMEINPUT_SUPPORTS_SENSORS
 	
 	EInputDeviceConnectionState DeviceStateToConnectionState(GameInputDeviceStatus InCurrentStatus, GameInputDeviceStatus InPreviousStatus)
 	{
diff --git a/Engine/Plugins/Runtime/GameInput/Source/GameInputBase/Private/IGameInputDeviceInterface.cpp b/Engine/Plugins/Runtime/GameInput/Source/GameInputBase/Private/IGameInputDeviceInterface.cpp
--- a/Engine/Plugins/Runtime/GameInput/Source/GameInputBase/Private/IGameInputDeviceInterface.cpp
+++ b/Engine/Plugins/Runtime/GameInput/Source/GameInputBase/Private/IGameInputDeviceInterface.cpp
@@ -307,6 +307,13 @@ GameInputKind IGameInputDeviceInterface::GetCurrentGameInputKindSupport() const
 		RegisterInputKindMask |= GameInputKindGamepad;
 	}
 
+#if UE_GAMEINPUT_SUPPORTS_SENSORS
+	if (PlatformSettings->bProcessSensors)
+	{
+		RegisterInputKindMask |= GameInputKindSensors;
+	}
+#endif
+	
 	if (PlatformSettings->bProcessKeyboard)
 	{
 		RegisterInputKindMask |= GameInputKindKeyboard;
diff --git a/Engine/Plugins/Runtime/GameInput/Source/GameInputBase/Public/GameInputBaseIncludes.h b/Engine/Plugins/Runtime/GameInput/Source/GameInputBase/Public/GameInputBaseIncludes.h
--- a/Engine/Plugins/Runtime/GameInput/Source/GameInputBase/Public/GameInputBaseIncludes.h
+++ b/Engine/Plugins/Runtime/GameInput/Source/GameInputBase/Public/GameInputBaseIncludes.h
@@ -25,6 +25,10 @@ THIRD_PARTY_INCLUDES_START
 
 	using namespace GameInput::v1;
 
+#elif GAMEINPUT_API_VERSION == 2
+
+	using namespace GameInput::v2;
+
 #endif
 
 // Depending on the Game Input API version it will support different types of input processing.
@@ -37,10 +41,25 @@ THIRD_PARTY_INCLUDES_START
 	#define UE_GAMEINPUT_SUPPORTS_RAW 1
 	#define UE_GAMEINPUT_SUPPORTS_TOUCH 1
 	#define UE_GAMEINPUT_SUPPORTS_DEVICE_STATUS 1
-#else
+	#define UE_GAMEINPUT_SUPPORTS_TRIGGER_HAPTICS 0
+	#define UE_GAMEINPUT_SUPPORTS_SENSORS 0
+	
+#elif GAMEINPUT_API_VERSION == 1 
+
+	#define UE_GAMEINPUT_SUPPORTS_RAW 0
+	#define UE_GAMEINPUT_SUPPORTS_TOUCH 0
+	#define UE_GAMEINPUT_SUPPORTS_DEVICE_STATUS 0
+	#define UE_GAMEINPUT_SUPPORTS_TRIGGER_HAPTICS 0
+	#define UE_GAMEINPUT_SUPPORTS_SENSORS 0
+	
+#else // version 2+
+
 	#define UE_GAMEINPUT_SUPPORTS_RAW 0
 	#define UE_GAMEINPUT_SUPPORTS_TOUCH 0
 	#define UE_GAMEINPUT_SUPPORTS_DEVICE_STATUS 0
+	#define UE_GAMEINPUT_SUPPORTS_TRIGGER_HAPTICS 1
+	#define UE_GAMEINPUT_SUPPORTS_SENSORS 1
+	
 #endif
 
 
diff --git a/Engine/Plugins/Runtime/GameInput/Source/GameInputBase/Public/GameInputDeveloperSettings.h b/Engine/Plugins/Runtime/GameInput/Source/GameInputBase/Public/GameInputDeveloperSettings.h
--- a/Engine/Plugins/Runtime/GameInput/Source/GameInputBase/Public/GameInputDeveloperSettings.h
+++ b/Engine/Plugins/Runtime/GameInput/Source/GameInputBase/Public/GameInputDeveloperSettings.h
@@ -356,6 +356,15 @@ public:
 	UPROPERTY(Config, EditAnywhere, Category="Processing Options")
 	bool bProcessMouse = true;
 
+	/**
+	 * If true, then we will process the GameInputKindSensors type such as Gyrometer, Orientation, and Accelerometer.
+	 *
+	 * Experimental.
+	 */
+	UE_EXPERIMENTAL(5.8, "Experimental, API may change in the future")
+	UPROPERTY(Config, EditAnywhere, Category="Processing Options")
+	bool bProcessSensors = true;
+
 	/**
 	 * If true, then we will process the GameInputKindRacingWheel type.
 	 * 
diff --git a/Engine/Plugins/Runtime/GameInput/Source/GameInputBase/Public/GameInputKeyTypes.h b/Engine/Plugins/Runtime/GameInput/Source/GameInputBase/Public/GameInputKeyTypes.h
--- a/Engine/Plugins/Runtime/GameInput/Source/GameInputBase/Public/GameInputKeyTypes.h
+++ b/Engine/Plugins/Runtime/GameInput/Source/GameInputBase/Public/GameInputKeyTypes.h
@@ -64,6 +64,12 @@ struct FGameInputKeys
 	static UE_API const FKey ArcadeStick_Action6;
 	static UE_API const FKey ArcadeStick_Special1;
 	static UE_API const FKey ArcadeStick_Special2;
+
+	//
+	// Sensor Types
+	//
+	UE_EXPERIMENTAL(5.8, "Experimental, API may change in the future")
+	static UE_API const FKey Sensor_HeadingInDegreesFromMagneticNorth;
 };
 
 #undef UE_API
diff --git a/Engine/Plugins/Runtime/GameInput/Source/GameInputBase/Public/GameInputUtils.h b/Engine/Plugins/Runtime/GameInput/Source/GameInputBase/Public/GameInputUtils.h
--- a/Engine/Plugins/Runtime/GameInput/Source/GameInputBase/Public/GameInputUtils.h
+++ b/Engine/Plugins/Runtime/GameInput/Source/GameInputBase/Public/GameInputUtils.h
@@ -19,6 +19,12 @@ namespace UE::GameInput
 
 	GAMEINPUTBASE_API FString LexToString(GameInputSwitchPosition SwitchPos);
 
+#if UE_GAMEINPUT_SUPPORTS_SENSORS
+	GAMEINPUTBASE_API FString LexToString(GameInputSensorsKind SensorKind);
+
+	GAMEINPUTBASE_API FString LexToString(GameInputSensorAccuracy SensorKind);
+#endif	// #if UE_GAMEINPUT_SUPPORTS_SENSORS
+
 	GAMEINPUTBASE_API FString LexToString(const APP_LOCAL_DEVICE_ID& DeviceId);
 
 	GAMEINPUTBASE_API const TCHAR* GetMouseButtonName(EMouseButtons::Type MouseButton);
@@ -67,6 +73,7 @@ namespace UE::GameInput
 	static constexpr float GamepadLeftStickDeadzone = (7849.0f / 32768.0f);
 	static constexpr float GamepadRightStickDeadzone = (8689.0f / 32768.0f);
 	static constexpr float GamepadTriggerDeadzone = (30.0f / 255.0f);
+	static constexpr float SensorMagneticNorthDeadzone = (30.0f / 255.0f);
 
 	static constexpr float InitialRepeatDelay = (0.2f);
 	static constexpr float SubsequentRepeatDelay = (0.1f);
