[Game Input] Add missing header piece that is compiled out behind UE_GAMEINPUT_SUPPORTS_SENSORS

#rnx

diff --git a/Engine/Plugins/Runtime/GameInput/Source/GameInputBase/Public/GameInputDeviceProcessor.h b/Engine/Plugins/Runtime/GameInput/Source/GameInputBase/Public/GameInputDeviceProcessor.h
--- a/Engine/Plugins/Runtime/GameInput/Source/GameInputBase/Public/GameInputDeviceProcessor.h
+++ b/Engine/Plugins/Runtime/GameInput/Source/GameInputBase/Public/GameInputDeviceProcessor.h
@@ -180,6 +180,34 @@ protected:
 	int32 NumReadingsProcessedThisFrame = 0;
 };
 
+#if UE_GAMEINPUT_SUPPORTS_SENSORS
+/**
+* Processor for Gamepad "Sensor" Inputs such as gyro or accelerometers
+*/
+UE_EXPERIMENTAL(5.8, "Experimental sensor data, API may change in the future")
+class FGameInputGamepadSensorDeviceProcessor final : public IGameInputDeviceProcessor
+{
+protected:
+	virtual bool ProcessInput(const FGameInputEventParams& Params) override;
+	virtual bool PostProcessInput(const FGameInputEventParams& Params) override;
+	virtual void ClearState(const FGameInputEventParams& Params) override;
+	virtual GameInputKind GetSupportedReadingKind() const override;
+	
+	/**
+	 * Process the current sensor state and report it to the engine via the message handler.
+	 * 
+	 * @param Params Current input param data
+	 * @param CurrentState The current sensor state for this frame
+	 * @return True if there was any sensor data reported.
+	 */
+	bool ProcessGamepadSensorState(const FGameInputEventParams& Params, GameInputSensorsState& CurrentState);
+	/**
+	 * The sensor data from the previous reading that was processed
+	 */
+	GameInputSensorsState PreviousState = {};
+};
+#endif// #if UE_GAMEINPUT_SUPPORTS_SENSORS
+
 struct FGameInputDeviceConfiguration;
 
 /**
