[Game Input] Add a version check to the editor and prompt the user to install if necessary.

You can disable this with the UE.GameInput.bEnableGameInputVersionCheck CVar

[RN] If you enable the "Game Input for Windows" plugin, we will automatically check that the user has the minimum version of GameInput available on their machine on editor boot.

diff --git a/Engine/Plugins/Runtime/GameInput/Source/GameInputBaseEditor/GameInputBaseEditor.build.cs b/Engine/Plugins/Runtime/GameInput/Source/GameInputBaseEditor/GameInputBaseEditor.build.cs
--- a/Engine/Plugins/Runtime/GameInput/Source/GameInputBaseEditor/GameInputBaseEditor.build.cs
+++ b/Engine/Plugins/Runtime/GameInput/Source/GameInputBaseEditor/GameInputBaseEditor.build.cs
@@ -12,7 +12,8 @@ public class GameInputBaseEditor : ModuleRules
 				"CoreUObject",
 				"Engine",
 				"GameInputBase",
-				"Settings"
+				"Settings",
+				"UnrealEd"
 			}
 		);
 	}
diff --git a/Engine/Plugins/Runtime/GameInput/Source/GameInputBaseEditor/Private/GameInputBaseEditorModule.cpp b/Engine/Plugins/Runtime/GameInput/Source/GameInputBaseEditor/Private/GameInputBaseEditorModule.cpp
--- a/Engine/Plugins/Runtime/GameInput/Source/GameInputBaseEditor/Private/GameInputBaseEditorModule.cpp
+++ b/Engine/Plugins/Runtime/GameInput/Source/GameInputBaseEditor/Private/GameInputBaseEditorModule.cpp
@@ -1,18 +1,97 @@
 // Copyright Epic Games, Inc. All Rights Reserved.
 
+#include "HAL/FileManager.h"
+#include "HAL/IConsoleManager.h"
+#include "HAL/PlatformProcess.h"
 #include "ISettingsModule.h"
+#include "Interfaces/IMainFrameModule.h"
+#include "Misc/App.h"
+#include "Misc/MessageDialog.h"
+#include "Misc/Paths.h"
 #include "Modules/ModuleInterface.h"
 #include "Modules/ModuleManager.h"
 #include "GameInputDeveloperSettings.h"
+#include "GameInputLogging.h"
+#include "UnrealEdMisc.h"
+#include "Dialogs/Dialogs.h"
+
+#if GAME_INPUT_SUPPORT
+#include "Windows/AllowWindowsPlatformTypes.h"
+THIRD_PARTY_INCLUDES_START
+
+#if PLATFORM_WINDOWS
+	#include <shellapi.h>
+#endif
+
+THIRD_PARTY_INCLUDES_END
+#include "Windows/HideWindowsPlatformTypes.h"
+#endif //#if GAME_INPUT_SUPPORT
 
 #define LOCTEXT_NAMESPACE "GameInputBaseEditor"
 
+namespace UE::Input
+{
+	static bool bEnableGameInputVersionCheck = true;
+	static FAutoConsoleVariableRef CVarEnableGameInputVersionCheck(
+		TEXT("UE.GameInput.bEnableGameInputVersionCheck"),
+		bEnableGameInputVersionCheck,
+		TEXT("If true, then we will check for a valid GameInput version when booting the editor and prompt users to install the GameInputRedist.msi if necessary"),
+		ECVF_Default);
+}
+
 /**
 * Editor module for Game Input that will register the UGameInputDeveloperSettings so they
 * show up as editable project settings in the editor.
 */
 class FGameInputBaseEditorModule final : public IModuleInterface
 {
+#if GAME_INPUT_SUPPORT
+	////////////////////////////////////////////////////////
+	// TODO: Some of these types should be moved out of Bootstrap and made available soon via the following header:	
+	// #include "Windows/WindowsRedistributableValidation.h"
+	// #jira UE-348535
+	
+#pragma region
+	/**
+	 * Visual C++ redistributable version information.
+	 */
+	struct VersionInfo
+	{
+		constexpr VersionInfo() = default;
+	
+		constexpr VersionInfo(unsigned long InMajor, unsigned long InMinor, unsigned long InBld, unsigned long InRbld)
+			: Major(InMajor)
+			, Minor(InMinor)
+			, Bld(InBld)
+			, Rbld(InRbld)
+		{}
+		constexpr VersionInfo(unsigned long long InVersion)
+			: Major((InVersion >> 48) & 0xffff)
+			, Minor((InVersion >> 32) & 0xffff)
+			, Bld((InVersion >> 16) & 0xffff)
+			, Rbld((InVersion >> 0) & 0xffff)
+		{ }
+		unsigned long Major = 0;
+		unsigned long Minor = 0;
+		unsigned long Bld = 0;
+		unsigned long Rbld = 0;
+	};
+	
+	/** Valid if the specified version is greater or equal than the minimum version. */
+	static bool IsVersionValid(const VersionInfo& Version, const VersionInfo& MinVersion)
+	{
+		if (Version.Major > MinVersion.Major) return true;
+		if (Version.Major == MinVersion.Major && Version.Minor > MinVersion.Minor) return true;
+		if (Version.Major == MinVersion.Major && Version.Minor == MinVersion.Minor && Version.Bld > MinVersion.Bld) return true;
+		if (Version.Major == MinVersion.Major && Version.Minor == MinVersion.Minor && Version.Bld == MinVersion.Bld && Version.Rbld >= MinVersion.Rbld) return true;
+		return false;
+	}
+
+#pragma endregion
+	// End of copying from the pending CL
+	//////////////////////////////////////////////////////
+#endif
+
 private:
 	virtual void StartupModule() override
 	{
@@ -25,6 +104,11 @@ private:
 				GetMutableDefault<UGameInputDeveloperSettings>()
 			);
 		}
+#if GAME_INPUT_SUPPORT
+		// Listen for when the editor is ready until we try and pop up a nice and friendly widget to upgrade
+		IMainFrameModule& MainFrameModule = FModuleManager::LoadModuleChecked<IMainFrameModule>(TEXT("MainFrame"));
+		MainFrameModule.OnMainFrameCreationFinished().AddRaw(this, &FGameInputBaseEditorModule::OnMainFrameCreationFinished);
+#endif
 	}
 
 	virtual void ShutdownModule() override
@@ -33,7 +117,237 @@ private:
 		{
 			SettingsModule->UnregisterSettings("Project", "Plugins", "Input");
 		}
+
+#if GAME_INPUT_SUPPORT
+		// Remove any listeners to the editor startup delegate
+		if (FModuleManager::Get().IsModuleLoaded("MainFrame"))
+		{
+			IMainFrameModule::Get().OnMainFrameCreationFinished().RemoveAll(this);
+		}
+#endif
+	}
+
+#if GAME_INPUT_SUPPORT
+	void OnMainFrameCreationFinished(TSharedPtr<SWindow> InRootWindow, bool bIsRunningStartupDialog)
+	{
+		if (FApp::HasProjectName())
+		{
+			ValidateGameInputAvailability();
+		}	
+	}
+
+	/**
+	 * Check if the user currently has the min version of GameInput installed.
+	 * 
+	 * If they don't, then prompt them to install it now
+	 */
+	void ValidateGameInputAvailability()
+	{
+		if (!UE::Input::bEnableGameInputVersionCheck)
+		{
+			return;
+		}
+		
+		if (HasValidVersionOfGameInput())
+		{
+			UE_LOG(LogGameInput, Log, TEXT("A valid Game Input installation was detected"));
+			return;
+		}
+
+		RunGameInputInstallerUserFlow();
+	}
+
+	/**
+	 * @return True if there is a valid version of Game Input currently available
+	 */
+	bool HasValidVersionOfGameInput()
+	{
+		static constexpr VersionInfo MinGameInputVersion = { 0, 2309, 26100, 5074 };
+		static constexpr VersionInfo MinGameInputRedistVersion = { 3, 0, 26100, 6154 };
+		
+		auto IsDllVersionValid = [](const TCHAR* ExecDirectory, const TCHAR* Name, const VersionInfo& MinRedistVersion) -> bool
+		{
+			if (!ExecDirectory)
+			{
+				return false;
+			}
+				
+			const FString FilePath = FPaths::Combine(ExecDirectory, Name);
+			if (!FilePath.IsEmpty())
+			{
+				const uint64 FileVersion = FWindowsPlatformMisc::GetFileVersion(FilePath);
+				if (FileVersion != 0)
+				{
+					const VersionInfo CurrentVersionInfo { FileVersion };
+					return IsVersionValid(CurrentVersionInfo, MinRedistVersion);
+				}
+			}
+			return false;
+		};
+
+		// Reproducing verifications done in BootstrapPackagedGame.cpp
+		bool bHasValidRedistributable = false;
+		
+		const FString SystemRoot = FWindowsPlatformMisc::GetEnvironmentVariable(TEXT("SystemRoot"));
+		const FString System32Path = FPaths::Combine(SystemRoot, TEXT("system32"));
+		
+		if (IsDllVersionValid(*System32Path, TEXT("GameInputRedist.dll"), MinGameInputRedistVersion) &&
+			IsDllVersionValid(*System32Path, TEXT("GameInput.dll"), MinGameInputVersion))
+		{
+			bHasValidRedistributable = true;
+		}
+		else if (IsDllVersionValid(nullptr, TEXT("GameInputRedist.dll"), MinGameInputRedistVersion) &&
+			IsDllVersionValid(nullptr, TEXT("GameInput.dll"), MinGameInputVersion))
+		{
+			bHasValidRedistributable = true;
+		}
+		
+		return bHasValidRedistributable;
+	}
+
+	bool PromptUserToInstallGameInput()
+	{
+		const FString GameInputInstallerPath = GetFullGameInputInstallerPath();
+
+		FFormatNamedArguments Args;
+		Args.Add(TEXT("InstallerPath"), FText::FromString(GameInputInstallerPath));
+		
+		// Pop up a widget with a yes/no prompt to run the Game Input installer
+		const FText RequiresInstallerMessageText =
+			FText::Format(LOCTEXT("GameInput_RequiresInstall_Prompt",
+			"A valid version of Game Input was not detected. Some input devices may not work as expected.\n\n"
+			"Would you like to run the installer now? An editor reboot will be required.\n\n"
+			"{InstallerPath}"),
+			Args);
+		
+		const FText Title = LOCTEXT("GameInput_RequiresInstall_Prompt_Title", "Missing Game Input Installation");
+
+		FSuppressableWarningDialog::FSetupInfo Info (RequiresInstallerMessageText, Title, TEXT("ShowGameInputRedistInstallPrompt"), GEditorPerProjectIni);
+		Info.ConfirmText = LOCTEXT("RunGameInputInstallerNow", "Install Now");
+		Info.CancelText = LOCTEXT("RunGameInputInstallerNow_Later", "No");
+		
+		FSuppressableWarningDialog RequireGameInputInstallPrompt(Info);
+		FSuppressableWarningDialog::EResult RunInstallerRes = RequireGameInputInstallPrompt.ShowModal();
+		
+		if (RunInstallerRes == FSuppressableWarningDialog::EResult::Cancel ||
+			RunInstallerRes == FSuppressableWarningDialog::EResult::Suppressed)
+		{
+			return false;
+		}
+		
+		return true;
+	}
+
+	static FString GetFullGameInputInstallerPath()
+	{
+		return FPaths::ConvertRelativePathToFull(FPaths::EngineDir() + TEXT("Extras/Redist/en-us/GameInputRedist.msi"));
+	}
+
+	/**
+	 * Actually executes the game input installer
+	 * @return The exit code of the Game input installer.
+	 */
+	uint32 ExecuteGameInputInstaller()
+	{
+		// Note: I would have liked to use FPlatformProcess::ExecProcess here, but because that adds the additional
+		// ShellExecuteInfo.lpVerb = TEXT("runas"); argument, it fails to run the installer MSI.
+		// So, fall back to the actual SHELLEXECUTEINFO API
+		
+#if PLATFORM_WINDOWS
+		const FString GameInputInstallerPath = GetFullGameInputInstallerPath();
+		
+		UE_LOG(LogGameInput, Log, TEXT("Starting the game input installer: %s"), *GameInputInstallerPath);
+		
+		// Start the game input installer
+		SHELLEXECUTEINFO ShellExecuteInfo;
+		ZeroMemory(&ShellExecuteInfo, sizeof(ShellExecuteInfo));
+		ShellExecuteInfo.cbSize = sizeof(ShellExecuteInfo);
+		ShellExecuteInfo.fMask = SEE_MASK_NOCLOSEPROCESS;
+		ShellExecuteInfo.nShow = SW_SHOWNORMAL;
+		ShellExecuteInfo.lpFile = *GameInputInstallerPath;
+		
+		if (!ShellExecuteExW(&ShellExecuteInfo))
+		{
+			UE_LOG(LogGameInput, Error, TEXT("Failed to run the game input installer: %s"), *GameInputInstallerPath);
+			return 9000;
+		}
+
+		// Wait for the process to complete, then get its exit code
+		WaitForSingleObject(ShellExecuteInfo.hProcess, INFINITE);
+
+		DWORD ProcessExitCode = 0;
+		GetExitCodeProcess(ShellExecuteInfo.hProcess, &ProcessExitCode);
+		CloseHandle(ShellExecuteInfo.hProcess);
+		
+		return (uint32)ProcessExitCode;
+#else
+		return 0;
+#endif	// #if PLATFORM_WINDOWS
+	}
+	
+	/**
+	 * Prompt the user to run the GameInputRedist.msi installer.
+	 */
+	void RunGameInputInstallerUserFlow()
+	{
+		// Ensure that we have a valid installer path before prompting users to run it
+		const FString GameInputInstallerPath = GetFullGameInputInstallerPath();
+
+		if (!IFileManager::Get().FileExists(*GameInputInstallerPath))
+		{
+			UE_LOG(LogGameInput, Error, TEXT("Failed to locate the Game Input Installer: %s"), *GameInputInstallerPath);
+			return;
+		}
+		
+		// The user has elected not to install game input
+		if (!PromptUserToInstallGameInput())
+		{
+			UE_LOG(LogGameInput, Warning, TEXT("Game Input was not detected, and the user has denied installation. Some input devices may not work."));
+			return;
+		}
+		
+		const uint32 ExitCode = ExecuteGameInputInstaller();
+		// 1638: Newer version already installed
+		const bool bSuccessfulInstall = ExitCode == 0 || ExitCode == 1638;
+
+		UE_LOG(LogGameInput, Display, TEXT("Game Input Installer returned exit code %d"), ExitCode);
+
+		// If we failed, then notify the user that some devices may not work and give them the error code
+		if (!bSuccessfulInstall)
+		{
+			const FText InstallFailedMessage =
+				FText::FormatOrdered(LOCTEXT("GameInputInstaller_Failed_Message",
+				"Game Input installation failed with error code {0}\n\nSome input devices may not behave as expected."),
+				ExitCode);
+
+			const FText InstallationSuccessMessageTitle =
+				LOCTEXT("GameInputInstaller_Failed_Title",
+				"Installation Failed");
+			
+			FMessageDialog::Open(EAppMsgCategory::Error, EAppMsgType::Ok, InstallFailedMessage, InstallationSuccessMessageTitle);
+			return;
+		}
+
+		// Otherwise the installation was a success, and in order to re-initialize Game Input we need to reboot the editor. Prompt the user to do so now
+		const FText RebootEditorPrompt =
+			LOCTEXT("GameInput_RebootEditor_Prompt",
+			"Game Input installation was a success! For the changes to take effect, you must restart the editor.\n\nWould you like to restart now?");
+
+		const FText InstallationSuccessMessageTitle =
+			LOCTEXT("GameInput_InstallSuccess_Title",
+			"Installation Successful");
+		
+		const EAppReturnType::Type RebootEditorRes = FMessageDialog::Open(EAppMsgCategory::Success, EAppMsgType::YesNo, RebootEditorPrompt, InstallationSuccessMessageTitle);
+
+		UE_CLOG(RebootEditorRes == EAppReturnType::No, LogGameInput, Warning, TEXT("Game Input was installed successfully, but without restarting the editor some input devices may not work."));
+		
+		if (RebootEditorRes == EAppReturnType::Yes)
+		{
+			constexpr bool bWarn = false;
+			FUnrealEdMisc::Get().RestartEditor(bWarn);
+		}
 	}
+#endif// #if GAME_INPUT_SUPPORT
 };
 
 IMPLEMENT_MODULE(FGameInputBaseEditorModule, GameInputBaseEditor);
