[Game Input] Do not check the "inbox" GameInput.dll version when bootstrapping the game, only the Redist.dll, to ensure that we have all the latest features.

#rnx

diff --git a/Engine/Plugins/Runtime/GameInput/Source/GameInputBaseEditor/Private/GameInputBaseEditorModule.cpp b/Engine/Plugins/Runtime/GameInput/Source/GameInputBaseEditor/Private/GameInputBaseEditorModule.cpp
--- a/Engine/Plugins/Runtime/GameInput/Source/GameInputBaseEditor/Private/GameInputBaseEditorModule.cpp
+++ b/Engine/Plugins/Runtime/GameInput/Source/GameInputBaseEditor/Private/GameInputBaseEditorModule.cpp
@@ -200,16 +200,10 @@ private:
 	 */
 	bool HasValidVersionOfGameInput()
 	{
-		static constexpr VersionInfo MinGameInputVersion = { 0, 2309, 26100, 5074 };
 		static constexpr VersionInfo MinGameInputRedistVersion = { 3, 0, 26100, 6154 };
 		
 		auto IsDllVersionValid = [](const TCHAR* ExecDirectory, const TCHAR* Name, const VersionInfo& MinRedistVersion) -> bool
 		{
-			if (!ExecDirectory)
-			{
-				return false;
-			}
-				
 			const FString FilePath = FPaths::Combine(ExecDirectory, Name);
 			if (!FilePath.IsEmpty())
 			{
@@ -224,23 +218,19 @@ private:
 		};
 
 		// Reproducing verifications done in BootstrapPackagedGame.cpp
-		bool bHasValidRedistributable = false;
+		bool bHasValidRedistInputDLL = false;
 		
 		const FString SystemRoot = FWindowsPlatformMisc::GetEnvironmentVariable(TEXT("SystemRoot"));
 		const FString System32Path = FPaths::Combine(SystemRoot, TEXT("system32"));
-		
-		if (IsDllVersionValid(*System32Path, TEXT("GameInputRedist.dll"), MinGameInputRedistVersion) &&
-			IsDllVersionValid(*System32Path, TEXT("GameInput.dll"), MinGameInputVersion))
-		{
-			bHasValidRedistributable = true;
-		}
-		else if (IsDllVersionValid(nullptr, TEXT("GameInputRedist.dll"), MinGameInputRedistVersion) &&
-			IsDllVersionValid(nullptr, TEXT("GameInput.dll"), MinGameInputVersion))
+
+		// Check GameInputRedist.dll
+		if (IsDllVersionValid(*System32Path, TEXT("GameInputRedist.dll"), MinGameInputRedistVersion) ||
+			IsDllVersionValid(nullptr, TEXT("GameInputRedist.dll"), MinGameInputRedistVersion))
 		{
-			bHasValidRedistributable = true;
+			bHasValidRedistInputDLL = true;
 		}
 		
-		return bHasValidRedistributable;
+		return bHasValidRedistInputDLL;
 	}
 
 	bool PromptUserToInstallGameInput()
diff --git a/Engine/Source/Programs/Windows/BootstrapPackagedGame/Private/BootstrapPackagedGame.cpp b/Engine/Source/Programs/Windows/BootstrapPackagedGame/Private/BootstrapPackagedGame.cpp
--- a/Engine/Source/Programs/Windows/BootstrapPackagedGame/Private/BootstrapPackagedGame.cpp
+++ b/Engine/Source/Programs/Windows/BootstrapPackagedGame/Private/BootstrapPackagedGame.cpp
@@ -20,7 +20,6 @@ struct VersionInfo
 // This minimum should match the version installed by
 // Engine/Extras/Redist/en-us/vc_redist.x64.exe
 static const VersionInfo MinRedistVersion = { 14, 44, 35211, 0 };
-static constexpr VersionInfo MinGameInputVersion = { 0, 2309, 26100, 5074 };
 static constexpr VersionInfo MinGameInputRedistVersion = { 3, 0, 26100, 6154 };
 
 bool IsVersionValid(const VersionInfo& Version, const VersionInfo& MinVersion)
@@ -212,8 +211,6 @@ bool HasGameInputInstalled()
 	// The min version should match the version installed by
 	// Engine/Extras/Redist/en-us/GameInputRedist.msi
 	
-	// We have two DLLs to version check for game input, GameInput and GameInputRedist.
-	bool bHasValidGameInput = false;
 	bool bHasValidGameInputRedist = false;
 
 	// Check the System32 folder for if the GameInput dll is available
@@ -221,19 +218,7 @@ bool HasGameInputInstalled()
 	GetEnvironmentVariable(L"SystemRoot", SystemRoot, MAX_PATH);
 	WCHAR System32Path[MAX_PATH] = { 0 };
 	PathCombine(System32Path, SystemRoot, L"system32");
-
-	// Check GameInput.dll verison 0.2309.26100.5074
-	{
-		if (IsDllValid(System32Path, L"GameInput.dll", MinGameInputVersion))
-		{
-			bHasValidGameInput = true;
-		}	
-		else if (IsDllValid(L"GameInput.dll", MinGameInputVersion))
-		{
-			bHasValidGameInput = true;
-		}
-	}
-
+	
 	// Check GameInputRedist.dll version 3.0.26100.6154
 	{
 		if (IsDllValid(System32Path, L"GameInputRedist.dll", MinGameInputRedistVersion))
@@ -246,7 +231,7 @@ bool HasGameInputInstalled()
 		}
 	}
 
-	return bHasValidGameInput && bHasValidGameInputRedist;
+	return bHasValidGameInputRedist;
 }
 
 /**
