[Game Input] Check the REDIST DLL version on boot at runtime instead of only in the editor to prevent a crash on binary builds which did not yet run the pre-req installer.

#rnx

diff --git a/Engine/Plugins/Runtime/GameInput/Source/GameInputBase/Private/GameInputBaseModule.cpp b/Engine/Plugins/Runtime/GameInput/Source/GameInputBase/Private/GameInputBaseModule.cpp
--- a/Engine/Plugins/Runtime/GameInput/Source/GameInputBase/Private/GameInputBaseModule.cpp
+++ b/Engine/Plugins/Runtime/GameInput/Source/GameInputBase/Private/GameInputBaseModule.cpp
@@ -6,6 +6,7 @@
 #include "CoreGlobals.h"
 #include "GameInputKeyTypes.h"
 #include "GameInputLogging.h"
+#include "GameInputUtils.h"
 #include "HAL/CriticalSection.h"
 #include "HAL/PlatformProcess.h"
 #include "Misc/App.h"
@@ -80,6 +81,13 @@ void FGameInputBaseModule::StartupModule()
 	UE_LOG(LogGameInput, Log, TEXT("Using GameInput Version: %d"), GAMEINPUT_API_VERSION);
 	
 #if PLATFORM_WINDOWS
+
+	if (!UE::GameInput::HasValidVersionOfGameInput())
+	{
+		UE_LOG(LogGameInput, Warning, TEXT("[%hs] module is exiting because a valid GameInput version is not installed. GameInput will not be initialized. Is it installed correctly?"), __func__);
+		return;
+	}
+	
 	// Check to see if the GameInput.dll exists...
 	// Search for the GameInput dll on desktop platforms. If for some reason it doesn't exist, then we shouldn't
 	// attempt to call any functions from it. The only known case for this is when running a client on a server OS
diff --git a/Engine/Plugins/Runtime/GameInput/Source/GameInputBase/Private/GameInputUtils.cpp b/Engine/Plugins/Runtime/GameInput/Source/GameInputBase/Private/GameInputUtils.cpp
--- a/Engine/Plugins/Runtime/GameInput/Source/GameInputBase/Private/GameInputUtils.cpp
+++ b/Engine/Plugins/Runtime/GameInput/Source/GameInputBase/Private/GameInputUtils.cpp
@@ -3,11 +3,98 @@
 #include "GameInputUtils.h"
 
 #include "GameInputDeveloperSettings.h"
+#include "Misc/Paths.h"
 
 #if GAME_INPUT_SUPPORT
 
 namespace UE::GameInput
 {
+#if GAME_INPUT_SUPPORT
+	////////////////////////////////////////////////////////
+	// TODO: Some of these types should be moved out of Bootstrap and made available soon via the following header:	
+	// #include "Windows/WindowsRedistributableValidation.h"
+	// #jira UE-348535
+	
+#pragma region
+	/**
+	 * Visual C++ redistributable version information.
+	 */
+	struct VersionInfo
+	{
+		constexpr VersionInfo() = default;
+	
+		constexpr VersionInfo(unsigned long InMajor, unsigned long InMinor, unsigned long InBld, unsigned long InRbld)
+			: Major(InMajor)
+			, Minor(InMinor)
+			, Bld(InBld)
+			, Rbld(InRbld)
+		{}
+		constexpr VersionInfo(unsigned long long InVersion)
+			: Major((InVersion >> 48) & 0xffff)
+			, Minor((InVersion >> 32) & 0xffff)
+			, Bld((InVersion >> 16) & 0xffff)
+			, Rbld((InVersion >> 0) & 0xffff)
+		{ }
+		unsigned long Major = 0;
+		unsigned long Minor = 0;
+		unsigned long Bld = 0;
+		unsigned long Rbld = 0;
+	};
+	
+	/** Valid if the specified version is greater or equal than the minimum version. */
+	static bool IsVersionValid(const VersionInfo& Version, const VersionInfo& MinVersion)
+	{
+		if (Version.Major > MinVersion.Major) return true;
+		if (Version.Major == MinVersion.Major && Version.Minor > MinVersion.Minor) return true;
+		if (Version.Major == MinVersion.Major && Version.Minor == MinVersion.Minor && Version.Bld > MinVersion.Bld) return true;
+		if (Version.Major == MinVersion.Major && Version.Minor == MinVersion.Minor && Version.Bld == MinVersion.Bld && Version.Rbld >= MinVersion.Rbld) return true;
+		return false;
+	}
+
+#pragma endregion
+	// End of copying from the pending CL
+	//////////////////////////////////////////////////////
+#endif
+	
+	bool HasValidVersionOfGameInput()
+	{
+#if PLATFORM_WINDOWS
+		static constexpr VersionInfo MinGameInputRedistVersion = { 3, 0, 26100, 6154 };
+		
+		auto IsDllVersionValid = [](const TCHAR* ExecDirectory, const TCHAR* Name, const VersionInfo& MinRedistVersion) -> bool
+			{
+				const FString FilePath = FPaths::Combine(ExecDirectory, Name);
+				if (!FilePath.IsEmpty())
+				{
+					const uint64 FileVersion = FWindowsPlatformMisc::GetFileVersion(FilePath);
+					if (FileVersion != 0)
+					{
+						const VersionInfo CurrentVersionInfo { FileVersion };
+						return IsVersionValid(CurrentVersionInfo, MinRedistVersion);
+					}
+				}
+				return false;
+			};
+
+		// Reproducing verifications done in BootstrapPackagedGame.cpp
+		bool bHasValidRedistInputDLL = false;
+		
+		const FString SystemRoot = FWindowsPlatformMisc::GetEnvironmentVariable(TEXT("SystemRoot"));
+		const FString System32Path = FPaths::Combine(SystemRoot, TEXT("system32"));
+
+		// Check GameInputRedist.dll
+		if (IsDllVersionValid(*System32Path, TEXT("GameInputRedist.dll"), MinGameInputRedistVersion) ||
+			IsDllVersionValid(nullptr, TEXT("GameInputRedist.dll"), MinGameInputRedistVersion))
+		{
+			bHasValidRedistInputDLL = true;
+		}
+		
+		return bHasValidRedistInputDLL;
+#else
+		return true;
+#endif
+	}
+
 	FString LexToString(IGameInputDevice* Device)
 	{
 		FString Result = TEXT("<invalid device>");
diff --git a/Engine/Plugins/Runtime/GameInput/Source/GameInputBase/Public/GameInputUtils.h b/Engine/Plugins/Runtime/GameInput/Source/GameInputBase/Public/GameInputUtils.h
--- a/Engine/Plugins/Runtime/GameInput/Source/GameInputBase/Public/GameInputUtils.h
+++ b/Engine/Plugins/Runtime/GameInput/Source/GameInputBase/Public/GameInputUtils.h
@@ -11,6 +11,11 @@
 
 namespace UE::GameInput
 {
+	/**
+	 * @return True if there is a valid version of Game Input currently available
+	 */
+	GAMEINPUTBASE_API bool HasValidVersionOfGameInput();
+	
 	GAMEINPUTBASE_API FString LexToString(IGameInputDevice* Device);
 	
 	GAMEINPUTBASE_API FString LexToString(GameInputDeviceStatus DeviceStatus);
diff --git a/Engine/Plugins/Runtime/GameInput/Source/GameInputBaseEditor/Private/GameInputBaseEditorModule.cpp b/Engine/Plugins/Runtime/GameInput/Source/GameInputBaseEditor/Private/GameInputBaseEditorModule.cpp
--- a/Engine/Plugins/Runtime/GameInput/Source/GameInputBaseEditor/Private/GameInputBaseEditorModule.cpp
+++ b/Engine/Plugins/Runtime/GameInput/Source/GameInputBaseEditor/Private/GameInputBaseEditorModule.cpp
@@ -1,5 +1,6 @@
 // Copyright Epic Games, Inc. All Rights Reserved.
 
+#include "GameInputUtils.h"
 #include "HAL/FileManager.h"
 #include "HAL/IConsoleManager.h"
 #include "HAL/PlatformProcess.h"
@@ -45,53 +46,6 @@ namespace UE::Input
 */
 class FGameInputBaseEditorModule final : public IModuleInterface
 {
-#if GAME_INPUT_SUPPORT
-	////////////////////////////////////////////////////////
-	// TODO: Some of these types should be moved out of Bootstrap and made available soon via the following header:	
-	// #include "Windows/WindowsRedistributableValidation.h"
-	// #jira UE-348535
-	
-#pragma region
-	/**
-	 * Visual C++ redistributable version information.
-	 */
-	struct VersionInfo
-	{
-		constexpr VersionInfo() = default;
-	
-		constexpr VersionInfo(unsigned long InMajor, unsigned long InMinor, unsigned long InBld, unsigned long InRbld)
-			: Major(InMajor)
-			, Minor(InMinor)
-			, Bld(InBld)
-			, Rbld(InRbld)
-		{}
-		constexpr VersionInfo(unsigned long long InVersion)
-			: Major((InVersion >> 48) & 0xffff)
-			, Minor((InVersion >> 32) & 0xffff)
-			, Bld((InVersion >> 16) & 0xffff)
-			, Rbld((InVersion >> 0) & 0xffff)
-		{ }
-		unsigned long Major = 0;
-		unsigned long Minor = 0;
-		unsigned long Bld = 0;
-		unsigned long Rbld = 0;
-	};
-	
-	/** Valid if the specified version is greater or equal than the minimum version. */
-	static bool IsVersionValid(const VersionInfo& Version, const VersionInfo& MinVersion)
-	{
-		if (Version.Major > MinVersion.Major) return true;
-		if (Version.Major == MinVersion.Major && Version.Minor > MinVersion.Minor) return true;
-		if (Version.Major == MinVersion.Major && Version.Minor == MinVersion.Minor && Version.Bld > MinVersion.Bld) return true;
-		if (Version.Major == MinVersion.Major && Version.Minor == MinVersion.Minor && Version.Bld == MinVersion.Bld && Version.Rbld >= MinVersion.Rbld) return true;
-		return false;
-	}
-
-#pragma endregion
-	// End of copying from the pending CL
-	//////////////////////////////////////////////////////
-#endif
-
 private:
 	virtual void StartupModule() override
 	{
@@ -186,7 +140,7 @@ private:
 			return;
 		}
 		
-		if (HasValidVersionOfGameInput())
+		if (UE::GameInput::HasValidVersionOfGameInput())
 		{
 			UE_LOG(LogGameInput, Log, TEXT("A valid Game Input installation was detected"));
 			return;
@@ -195,44 +149,6 @@ private:
 		RunGameInputInstallerUserFlow();
 	}
 
-	/**
-	 * @return True if there is a valid version of Game Input currently available
-	 */
-	bool HasValidVersionOfGameInput()
-	{
-		static constexpr VersionInfo MinGameInputRedistVersion = { 3, 0, 26100, 6154 };
-		
-		auto IsDllVersionValid = [](const TCHAR* ExecDirectory, const TCHAR* Name, const VersionInfo& MinRedistVersion) -> bool
-		{
-			const FString FilePath = FPaths::Combine(ExecDirectory, Name);
-			if (!FilePath.IsEmpty())
-			{
-				const uint64 FileVersion = FWindowsPlatformMisc::GetFileVersion(FilePath);
-				if (FileVersion != 0)
-				{
-					const VersionInfo CurrentVersionInfo { FileVersion };
-					return IsVersionValid(CurrentVersionInfo, MinRedistVersion);
-				}
-			}
-			return false;
-		};
-
-		// Reproducing verifications done in BootstrapPackagedGame.cpp
-		bool bHasValidRedistInputDLL = false;
-		
-		const FString SystemRoot = FWindowsPlatformMisc::GetEnvironmentVariable(TEXT("SystemRoot"));
-		const FString System32Path = FPaths::Combine(SystemRoot, TEXT("system32"));
-
-		// Check GameInputRedist.dll
-		if (IsDllVersionValid(*System32Path, TEXT("GameInputRedist.dll"), MinGameInputRedistVersion) ||
-			IsDllVersionValid(nullptr, TEXT("GameInputRedist.dll"), MinGameInputRedistVersion))
-		{
-			bHasValidRedistInputDLL = true;
-		}
-		
-		return bHasValidRedistInputDLL;
-	}
-
 	bool PromptUserToInstallGameInput()
 	{
 		const FString GameInputInstallerPath = GetFullGameInputInstallerPath();
