[Game Input] Remove the now unnecessary DLL check in favor of the more robust UE::GameInput::HasValidVersionOfGameInput call which was introduced with 46936146.

This is no longer necessary because the call to HasValidVersionOfGameInput will fail and actually check for the correct version of a GameInput.dll on the users' machine before ever attempting to use it. This way we can enfore a minimum version to support all the features that we want from game input, and also still be able to check for the DLLs existence on platforms like Windows Server, which may be used to run automation tests, that don't have a GameInput.dll installed at all (and you dont' want it to, because it is all headless automations)

Also PF with extra cooking tests: 6904df8dbb31b7e9e867dfbf

diff --git a/Engine/Plugins/Runtime/GameInput/Source/GameInputBase/GameInputBase.build.cs b/Engine/Plugins/Runtime/GameInput/Source/GameInputBase/GameInputBase.build.cs
--- a/Engine/Plugins/Runtime/GameInput/Source/GameInputBase/GameInputBase.build.cs
+++ b/Engine/Plugins/Runtime/GameInput/Source/GameInputBase/GameInputBase.build.cs
@@ -6,9 +6,6 @@ namespace UnrealBuildTool.Rules
 {
 	public class GameInputBase : ModuleRules
 	{
-		// If true, then we will use the dynamic GameInput.dll file instead of a statically linked GameInput library.
-		protected virtual bool bRequiresDynamicCreation => false;
-
 		/**
 		 * Conditions which are required for this build target to have support for the GameInput library.
 		 */
@@ -45,12 +42,7 @@ namespace UnrealBuildTool.Rules
 			
 			// Define this as 0 in the base module to avoid compilation errors when building
 			// without any Game Input support. It is up to the platform-specific submodules to define 
-			PublicDefinitions.Add("GAME_INPUT_SUPPORT=" + (HasGameInputSupport ? "1" : "0"));
-			
-			// Some targets may want to disable this plugin entirely
-			bool bShouldForceDynamicCreation = Target.GlobalDefinitions.Contains("UE_GAME_INPUT_FORCE_DYNAMIC_CREATION=1");
-			
-			PublicDefinitions.Add("UE_GAMEINPUT_REQUIRES_DYNAMIC_CREATION=" + (bShouldForceDynamicCreation || bRequiresDynamicCreation ? "1" : "0"));
+			PublicDefinitions.Add("GAME_INPUT_SUPPORT=" + (HasGameInputSupport ? "1" : "0"));			
 
 			// Give platforms extensions a chance to add any required dependencies that may be necessary to compile game input
 			if (HasGameInputSupport)
diff --git a/Engine/Plugins/Runtime/GameInput/Source/GameInputBase/Private/GameInputBaseModule.cpp b/Engine/Plugins/Runtime/GameInput/Source/GameInputBase/Private/GameInputBaseModule.cpp
--- a/Engine/Plugins/Runtime/GameInput/Source/GameInputBase/Private/GameInputBaseModule.cpp
+++ b/Engine/Plugins/Runtime/GameInput/Source/GameInputBase/Private/GameInputBaseModule.cpp
@@ -21,12 +21,6 @@ namespace UE::GameInput
 	// A singleton pointer to the base GameInput interface. 
 	// This provides access to reading the input stream, device callbacks, and more. 
 	static TComPtr<IGameInput> GGameInputInterface;
-
-#if PLATFORM_WINDOWS
-	// The name of the Game Input DLL file which would be needed when running on windows
-	static const FString GameInputDLLPath = TEXT("GameInput.dll");
-#endif	// #if PLATFORM_WINDOWS
-
 #endif	// #if GAME_INPUT_SUPPORT	
 }
 
@@ -87,49 +81,18 @@ void FGameInputBaseModule::StartupModule()
 		UE_LOG(LogGameInput, Warning, TEXT("[%hs] module is exiting because a valid GameInput version is not installed. GameInput will not be initialized. Is it installed correctly?"), __func__);
 		return;
 	}
-	
-	// Check to see if the GameInput.dll exists...
-	// Search for the GameInput dll on desktop platforms. If for some reason it doesn't exist, then we shouldn't
-	// attempt to call any functions from it. The only known case for this is when running a client on a server OS
-	// which doesn't have game input installed by default. 
-	GameInputDLLHandle = FPlatformProcess::GetDllHandle(*UE::GameInput::GameInputDLLPath);
-	if (GameInputDLLHandle == nullptr)
-	{
-		UE_LOG(LogGameInput, Warning, TEXT("[%hs] module is exiting because '%s' cannot be found. GameInput will not be initialized. Is it installed correctly?"), __func__, *UE::GameInput::GameInputDLLPath);
-		return;
-	}
 
 // Disable warning C4191: 'type cast' : unsafe conversion from 'PROC' to 'XXX' getting the GameInputCreate function
 #pragma warning(push)
 #pragma warning(disable:4191)
 
 	// The GameInputCreate function can be expensive on windows (~5s on startup!) so run it in an async task
-	UE::Tasks::Launch(UE_SOURCE_LOCATION, [DllHandle = GameInputDLLHandle, FinishedDelegate = MakeShared<TMulticastDelegate<void(IGameInput*)>>(OnGameInputCreation)]()
+	UE::Tasks::Launch(UE_SOURCE_LOCATION, [FinishedDelegate = MakeShared<TMulticastDelegate<void(IGameInput*)>>(OnGameInputCreation)]()
 	{		
-		// Create the Game Input interface via the dynamic DLL function call.
-		typedef HRESULT(WINAPI* FGameInputCreateFn)(IGameInput**);
-		FGameInputCreateFn pfnGameInputCreate = (FGameInputCreateFn)GetProcAddress((HMODULE)DllHandle, "GameInputCreate");
-
-		// We need to check that we have a valid GameInputCreate function available on the system, otherwise it will crash the application. The only known
-		// case for this is on Windows Server targets and some virtualized containers. 
-		if (!pfnGameInputCreate)
-		{
-			UE_LOG(LogGameInput, Warning, TEXT("[%hs] Failed to GetProcAddress (GameInputCreate). Game Input will fail to be created."), __func__);
-			return;
-		}
-		
 		GameInputCreationLock.Lock();
 
-#ifndef UE_GAME_INPUT_FORCE_DYNAMIC_CREATION
-	#define UE_GAME_INPUT_FORCE_DYNAMIC_CREATION 0
-#endif
-			
-#if UE_GAMEINPUT_REQUIRES_DYNAMIC_CREATION || UE_GAME_INPUT_FORCE_DYNAMIC_CREATION	
-		const HRESULT HResult = pfnGameInputCreate ? pfnGameInputCreate(&UE::GameInput::GGameInputInterface) : E_FAIL;
-#else
 		const HRESULT HResult = GameInputCreate(&UE::GameInput::GGameInputInterface);
-#endif
-			
+		
 		GameInputCreationLock.Unlock();
 		
 		if (SUCCEEDED(HResult))
@@ -205,14 +168,7 @@ void FGameInputBaseModule::ShutdownModule()
 #endif
 
 	// GGameInputInterface must be reset before we release the DLL handle
-	UE::GameInput::GGameInputInterface.Reset();
-	
-#if PLATFORM_WINDOWS
-	if (GameInputDLLHandle)
-	{
-		FPlatformProcess::FreeDllHandle(GameInputDLLHandle);
-	}
-#endif // #if PLATFORM_WINDOWS
+	UE::GameInput::GGameInputInterface.Reset();	
 
 #endif // #if GAME_INPUT_SUPPORT
 }
diff --git a/Engine/Plugins/Runtime/GameInput/Source/GameInputBase/Public/GameInputBaseModule.h b/Engine/Plugins/Runtime/GameInput/Source/GameInputBase/Public/GameInputBaseModule.h
--- a/Engine/Plugins/Runtime/GameInput/Source/GameInputBase/Public/GameInputBaseModule.h
+++ b/Engine/Plugins/Runtime/GameInput/Source/GameInputBase/Public/GameInputBaseModule.h
@@ -41,12 +41,7 @@ protected:
 	UE_API void InitializeGameInputKeys();
 
 #if PLATFORM_WINDOWS && GAME_INPUT_SUPPORT
-	// Handle to the game input dll which is set on StartupModule.
-	// If we can't find the DLL then we will early exit and not attempt to initalize GameInput.
-	void* GameInputDLLHandle = nullptr;
-
 	static UE_API FCriticalSection GameInputCreationLock;
-	
 #endif // endif PLATFORM_WINDOWS && GAME_INPUT_SUPPORT
 };
 
