[Game Input] Log the version number being used for GameInput on boot.

#rnx

diff --git a/Engine/Plugins/Runtime/GameInput/Source/GameInputBase/Private/GameInputBaseModule.cpp b/Engine/Plugins/Runtime/GameInput/Source/GameInputBase/Private/GameInputBaseModule.cpp
--- a/Engine/Plugins/Runtime/GameInput/Source/GameInputBase/Private/GameInputBaseModule.cpp
+++ b/Engine/Plugins/Runtime/GameInput/Source/GameInputBase/Private/GameInputBaseModule.cpp
@@ -72,7 +72,7 @@ void FGameInputBaseModule::StartupModule()
 
 #if GAME_INPUT_SUPPORT
 
-	UE_LOG(LogGameInput, Log, TEXT("Using GameInput Version: %d"), GAMEINPUT_API_VERSION);
+	UE_LOG(LogGameInput, Log, TEXT("Compiled using GameInput major Version: %d"), GAMEINPUT_API_VERSION);
 	
 #if PLATFORM_WINDOWS
 
@@ -82,6 +82,8 @@ void FGameInputBaseModule::StartupModule()
 		return;
 	}
 
+	UE_LOG(LogGameInput, Log, TEXT("Executing with detected GameInput Version: %s"), *UE::GameInput::GetInstalledGameInputVersionInfoString());
+
 // Disable warning C4191: 'type cast' : unsafe conversion from 'PROC' to 'XXX' getting the GameInputCreate function
 #pragma warning(push)
 #pragma warning(disable:4191)
@@ -166,8 +168,7 @@ void FGameInputBaseModule::ShutdownModule()
 #if PLATFORM_WINDOWS
 	FScopeLock Lock(&GameInputCreationLock);
 #endif
-
-	// GGameInputInterface must be reset before we release the DLL handle
+	
 	UE::GameInput::GGameInputInterface.Reset();	
 
 #endif // #if GAME_INPUT_SUPPORT
diff --git a/Engine/Plugins/Runtime/GameInput/Source/GameInputBase/Private/GameInputUtils.cpp b/Engine/Plugins/Runtime/GameInput/Source/GameInputBase/Private/GameInputUtils.cpp
--- a/Engine/Plugins/Runtime/GameInput/Source/GameInputBase/Private/GameInputUtils.cpp
+++ b/Engine/Plugins/Runtime/GameInput/Source/GameInputBase/Private/GameInputUtils.cpp
@@ -13,6 +13,61 @@
 
 namespace UE::GameInput
 {	
+	/**
+	* Returns the version info of the valid GameInputRedist.dll if one is available. 
+	*/
+	VersionInfo GetInstalledGameInputVersionInfo()
+	{		 		
+#if PLATFORM_WINDOWS
+		const FString SystemRoot = FWindowsPlatformMisc::GetEnvironmentVariable(TEXT("SystemRoot"));
+    	const FString System32Path = FPaths::Combine(SystemRoot, TEXT("system32"));
+    
+		auto GetDLLVersionInfo = [](const TCHAR* ExecDirectory, const TCHAR* Name) -> VersionInfo
+		{
+			const FString FilePath = FPaths::Combine(ExecDirectory, Name);
+			if (!FilePath.IsEmpty())
+			{
+				const uint64 FileVersion = FWindowsPlatformMisc::GetFileVersion(FilePath);
+				return VersionInfo{ FileVersion };
+			}
+			return false;
+		};
+		
+		// We always prefer the version in System32 it is valid.
+		const VersionInfo System32Version = GetDLLVersionInfo(*System32Path, TEXT("GameInputRedist.dll"));
+		if (IsVersionValid(System32Version, MinGameInputRedistVersion))
+		{
+			return System32Version;
+		}
+		
+		// Otherwise we fall back to the local dll
+		const VersionInfo LocalVersion = GetDLLVersionInfo(nullptr, TEXT("GameInputRedist.dll"));
+		if (IsVersionValid(LocalVersion, MinGameInputRedistVersion))
+		{
+			return LocalVersion;
+		}
+#endif
+    		    		   
+		// Non-Windows platforms won't have a GameInputRedist.dll, so just return empty version info 	
+		return VersionInfo();
+	}
+
+	FString GetInstalledGameInputVersionInfoString()
+	{
+		const VersionInfo InstalledVersion = UE::GameInput::GetInstalledGameInputVersionInfo();
+		
+		return FString::Printf(
+			TEXT("%lu.%lu.%lu.%lu"), 
+			InstalledVersion.Major, InstalledVersion.Minor, InstalledVersion.Bld, InstalledVersion.Rbld);
+	}
+	
+	FString GetMinVersionOfGameInputRedistString()
+	{
+		return FString::Printf(
+    		TEXT("%lu.%lu.%lu.%lu"), 
+    		MinGameInputRedistVersion.Major, MinGameInputRedistVersion.Minor, MinGameInputRedistVersion.Bld, MinGameInputRedistVersion.Rbld);
+	}
+
 	bool HasValidVersionOfGameInput()
 	{
 #if PLATFORM_WINDOWS		
diff --git a/Engine/Plugins/Runtime/GameInput/Source/GameInputBase/Public/GameInputUtils.h b/Engine/Plugins/Runtime/GameInput/Source/GameInputBase/Public/GameInputUtils.h
--- a/Engine/Plugins/Runtime/GameInput/Source/GameInputBase/Public/GameInputUtils.h
+++ b/Engine/Plugins/Runtime/GameInput/Source/GameInputBase/Public/GameInputUtils.h
@@ -16,6 +16,16 @@ namespace UE::GameInput
 	 */
 	GAMEINPUTBASE_API bool HasValidVersionOfGameInput();
 	
+	/**
+	* Returns a string which represents the version of Game Input that is installed on this machine. 
+	*/
+	GAMEINPUTBASE_API FString GetInstalledGameInputVersionInfoString();
+	
+	/**
+	* Returns a string which represents the min required version of Game Input needed to run the engine.
+	*/
+	GAMEINPUTBASE_API FString GetMinVersionOfGameInputRedistString();
+	
 	GAMEINPUTBASE_API FString LexToString(IGameInputDevice* Device);
 	
 	GAMEINPUTBASE_API FString LexToString(GameInputDeviceStatus DeviceStatus);
diff --git a/Engine/Plugins/Runtime/GameInput/Source/GameInputBaseEditor/Private/GameInputBaseEditorModule.cpp b/Engine/Plugins/Runtime/GameInput/Source/GameInputBaseEditor/Private/GameInputBaseEditorModule.cpp
--- a/Engine/Plugins/Runtime/GameInput/Source/GameInputBaseEditor/Private/GameInputBaseEditorModule.cpp
+++ b/Engine/Plugins/Runtime/GameInput/Source/GameInputBaseEditor/Private/GameInputBaseEditorModule.cpp
@@ -141,8 +141,8 @@ private:
 		}
 		
 		if (UE::GameInput::HasValidVersionOfGameInput())
-		{
-			UE_LOG(LogGameInput, Log, TEXT("A valid Game Input installation was detected"));
+		{			
+			UE_LOG(LogGameInput, Log, TEXT("A valid Game Input installation was detected: %s"), *UE::GameInput::GetInstalledGameInputVersionInfoString());
 			return;
 		}
 
@@ -155,12 +155,14 @@ private:
 
 		FFormatNamedArguments Args;
 		Args.Add(TEXT("InstallerPath"), FText::FromString(GameInputInstallerPath));
+		Args.Add(TEXT("DetectedVersionNum"), FText::FromString(UE::GameInput::GetInstalledGameInputVersionInfoString()));
+		Args.Add(TEXT("MinVersionNum"), FText::FromString(UE::GameInput::GetMinVersionOfGameInputRedistString()));
 		
 		// Pop up a widget with a yes/no prompt to run the Game Input installer
 		const FText RequiresInstallerMessageText =
 			FText::Format(LOCTEXT("GameInput_RequiresInstall_Prompt",
-			"A valid version of Game Input was not detected. Some input devices may not work as expected.\n\n"
-			"Would you like to run the installer now? An editor reboot will be required.\n\n"
+			"A valid version of Game Input was not detected ({DetectedVersionNum}). Some input devices may not work as expected.\n\n"
+			"Would you like to run the Game Input {MinVersionNum} installer now? An editor reboot will be required.\n\n"
 			"{InstallerPath}"),
 			Args);
 		
