[Game Input] Add support for "Dynamic Latency Input" on plaforms which support it.

#nosubmit
#rnx
#tests Run packaged lyra on platforms which support DLI and ensure that the controllers behave as expected, as well as keyboard
#Synced-CL 48353899

diff --git a/Engine/Plugins/Runtime/GameInput/Source/GameInputBase/GameInputBase.build.cs b/Engine/Plugins/Runtime/GameInput/Source/GameInputBase/GameInputBase.build.cs
--- a/Engine/Plugins/Runtime/GameInput/Source/GameInputBase/GameInputBase.build.cs
+++ b/Engine/Plugins/Runtime/GameInput/Source/GameInputBase/GameInputBase.build.cs
@@ -11,6 +11,11 @@ namespace UnrealBuildTool.Rules
 		 */
 		protected virtual bool HasGameInputSupport => Target.Platform.IsInGroup(UnrealPlatformGroup.Windows);
 
+		/**
+		* True on platforms which support the "Dynamic Latency Input" API of Game Input.
+		*/
+		protected virtual bool bSupportsDynamicLatencyInput => false;
+
 		/**
 		 * Array of required module names to be added to PublicDependencyModuleNames
 		 * if the target supports game input
@@ -44,6 +49,8 @@ namespace UnrealBuildTool.Rules
 			// without any Game Input support. It is up to the platform-specific submodules to define 
 			PublicDefinitions.Add("GAME_INPUT_SUPPORT=" + (HasGameInputSupport ? "1" : "0"));			
 
+			PublicDefinitions.Add("UE_GAME_INPUT_SUPPORTS_DLI=" + (bSupportsDynamicLatencyInput ? "1" : "0"));
+			
 			// Give platforms extensions a chance to add any required dependencies that may be necessary to compile game input
 			if (HasGameInputSupport)
 			{
diff --git a/Engine/Plugins/Runtime/GameInput/Source/GameInputBase/Private/GameInputDeviceContainer.cpp b/Engine/Plugins/Runtime/GameInput/Source/GameInputBase/Private/GameInputDeviceContainer.cpp
--- a/Engine/Plugins/Runtime/GameInput/Source/GameInputBase/Private/GameInputDeviceContainer.cpp
+++ b/Engine/Plugins/Runtime/GameInput/Source/GameInputBase/Private/GameInputDeviceContainer.cpp
@@ -205,6 +205,25 @@ const GameInputKind FGameInputDeviceContainer::ProcessInput(IGameInput* GameInpu
 
 	const bool bDoNotProcessDuplicateCapabilitiesForSingleUser = GetDefault<UGameInputDeveloperSettings>()->bDoNotProcessDuplicateCapabilitiesForSingleUser;
 
+#if UE_GAME_INPUT_SUPPORTS_DLI
+	// We can send a Synchronization Hint to Game Input to make possible improvements to input latency with
+	// Dynamic Latency Input (DLI) interface. The docs say to send this hint just before the first read of any input
+	// on the device for the simulation, which will be directly below in the "GetCurrentReading" call.
+	// 
+	// For supported devices, DLI will then attempt to sync the reporting rate of the controller 
+	// with that of the game simulation, providing better input latency.
+	//
+	// Note: If you call this at higher than 125hz, the DLI interface will disable itself
+	{
+		const GameInputDeviceInfo* Info = UE::GameInput::GetDeviceInfo(Device);
+		
+		if (Info && (Info->capabilities & GameInputDeviceCapabilitySynchronization) != 0)
+		{
+			Device->SendInputSynchronizationHint();
+		}
+	}
+#endif	// #if UE_GAME_INPUT_SUPPORTS_DLI
+	
 	// keep reading the input snapshots for this device
 	if (!LastReading.IsValid())
 	{				
