[GameInput] Mark some settings in the developer settings as requiring an editor restart.

#rnx
#Synced-CL 48353899

diff --git a/Engine/Plugins/Runtime/GameInput/Source/GameInputBase/Public/GameInputDeveloperSettings.h b/Engine/Plugins/Runtime/GameInput/Source/GameInputBase/Public/GameInputDeveloperSettings.h
--- a/Engine/Plugins/Runtime/GameInput/Source/GameInputBase/Public/GameInputDeveloperSettings.h
+++ b/Engine/Plugins/Runtime/GameInput/Source/GameInputBase/Public/GameInputDeveloperSettings.h
@@ -216,14 +216,14 @@ struct FGameInputDeviceConfiguration
 	* it is really some specially manufactured hardware for your game that you would like 
 	* to have access to at the Gameplay layer in UE.
 	*/
-	UPROPERTY(EditAnywhere, Config, Category = "Device Settings", meta = (InlineEditConditionToggle))
+	UPROPERTY(EditAnywhere, Config, Category = "Device Settings", meta = (InlineEditConditionToggle, ConfigRestartRequired=true))
 	bool bOverrideHardwareDeviceIdString;
 
 	/**
 	* The name of this device that should be used for FInputDeviceScope's to determine 
 	* when it has been connected or input events came from it.
 	*/
-	UPROPERTY(EditAnywhere, Config, Category = "Device Settings", meta=(EditCondition="bOverrideHardwareDeviceIdString"))
+	UPROPERTY(EditAnywhere, Config, Category = "Device Settings", meta=(EditCondition="bOverrideHardwareDeviceIdString", ConfigRestartRequired=true))
 	FString OverriddenHardwareDeviceId;
 
 	/** If true, this device will attempt to process Game Input Button states. */
