[Game Input] Refactor the device processors to be in seperate header files.

Deprecated "GameInputDeviceProcessor.h" in favor of the broken up "Processors/GameInputDeviceProcessor_*.h" files.

#rnx
#Synced-CL 48353899

diff --git a/Engine/Plugins/Runtime/GameInput/Source/GameInputBase/Private/GameInputDeviceContainer.cpp b/Engine/Plugins/Runtime/GameInput/Source/GameInputBase/Private/GameInputDeviceContainer.cpp
--- a/Engine/Plugins/Runtime/GameInput/Source/GameInputBase/Private/GameInputDeviceContainer.cpp
+++ b/Engine/Plugins/Runtime/GameInput/Source/GameInputBase/Private/GameInputDeviceContainer.cpp
@@ -7,6 +7,16 @@
 #include "GameInputLogging.h"
 #include "GameInputDeveloperSettings.h"
 #include "HAL/IConsoleManager.h"
+#include "Processors/IGameInputDeviceProcessor.h"
+#include "Processors/GameInputDeviceProcessor_Touch.h"
+#include "Processors/GameInputDeviceProcessor_Sensor.h"
+#include "Processors/GameInputDeviceProcessor_Raw.h"
+#include "Processors/GameInputDeviceProcessor_RacingWheel.h"
+#include "Processors/GameInputDeviceProcessor_KeyboardMouse.h"
+#include "Processors/GameInputDeviceProcessor_Gamepad.h"
+#include "Processors/GameInputDeviceProcessor_FlightStick.h"
+#include "Processors/GameInputDeviceProcessor_Controller.h"
+#include "Processors/GameInputDeviceProcessor_ArcadeStick.h"
 
 #if GAME_INPUT_SUPPORT
 
diff --git a/Engine/Plugins/Runtime/GameInput/Source/GameInputBase/Private/GameInputDeviceProcessor.cpp b/Engine/Plugins/Runtime/GameInput/Source/GameInputBase/Private/GameInputDeviceProcessor.cpp
deleted file mode 100644
diff --git a/Engine/Plugins/Runtime/GameInput/Source/GameInputBase/Private/Processors/GameInputDeviceProcessor_ArcadeStick.cpp b/Engine/Plugins/Runtime/GameInput/Source/GameInputBase/Private/Processors/GameInputDeviceProcessor_ArcadeStick.cpp
new file mode 100644
--- /dev/null
+++ b/Engine/Plugins/Runtime/GameInput/Source/GameInputBase/Private/Processors/GameInputDeviceProcessor_ArcadeStick.cpp
@@ -0,0 +1,111 @@
+// Copyright Epic Games, Inc. All Rights Reserved.
+
+#include "Processors/GameInputDeviceProcessor_ArcadeStick.h"
+
+#include "GameInputDeveloperSettings.h"
+#include "GameInputKeyTypes.h"
+
+#if GAME_INPUT_SUPPORT
+
+namespace UE::GameInput
+{
+	/** A map of uint32 GameInput button bitmask flags to the associated Unreal Engine FKey name. */
+	static const TMap<uint32, FGamepadKeyNames::Type>& GetArcadeStickButtonMap()
+	{
+		static const TMap<uint32, FGamepadKeyNames::Type> GamepadButtonMap
+		{
+			// Generic gamepad buttons
+			{ static_cast<uint32>(GameInputArcadeStickButtons::GameInputArcadeStickNone), FGamepadKeyNames::Invalid},
+			{ static_cast<uint32>(GameInputArcadeStickButtons::GameInputArcadeStickMenu), FGamepadKeyNames::SpecialRight },
+			{ static_cast<uint32>(GameInputArcadeStickButtons::GameInputArcadeStickView), FGamepadKeyNames::SpecialLeft },
+			{ static_cast<uint32>(GameInputArcadeStickButtons::GameInputArcadeStickUp), FGamepadKeyNames::DPadUp },
+			{ static_cast<uint32>(GameInputArcadeStickButtons::GameInputArcadeStickDown), FGamepadKeyNames::DPadDown },
+			{ static_cast<uint32>(GameInputArcadeStickButtons::GameInputArcadeStickLeft), FGamepadKeyNames::DPadLeft },
+			{ static_cast<uint32>(GameInputArcadeStickButtons::GameInputArcadeStickRight), FGamepadKeyNames::DPadRight },
+
+			// Unique to arcade sticks
+			{ static_cast<uint32>(GameInputArcadeStickButtons::GameInputArcadeStickAction1), FGameInputKeys::ArcadeStick_Action1.GetFName() },
+			{ static_cast<uint32>(GameInputArcadeStickButtons::GameInputArcadeStickAction2), FGameInputKeys::ArcadeStick_Action2.GetFName() },
+			{ static_cast<uint32>(GameInputArcadeStickButtons::GameInputArcadeStickAction3), FGameInputKeys::ArcadeStick_Action3.GetFName() },
+			{ static_cast<uint32>(GameInputArcadeStickButtons::GameInputArcadeStickAction4), FGameInputKeys::ArcadeStick_Action4.GetFName() },
+			{ static_cast<uint32>(GameInputArcadeStickButtons::GameInputArcadeStickAction5), FGameInputKeys::ArcadeStick_Action5.GetFName() },
+			{ static_cast<uint32>(GameInputArcadeStickButtons::GameInputArcadeStickAction6), FGameInputKeys::ArcadeStick_Action6.GetFName() },
+			{ static_cast<uint32>(GameInputArcadeStickButtons::GameInputArcadeStickSpecial1), FGameInputKeys::ArcadeStick_Special1.GetFName() },
+			{ static_cast<uint32>(GameInputArcadeStickButtons::GameInputArcadeStickSpecial2), FGameInputKeys::ArcadeStick_Special2.GetFName() }
+		};
+		return GamepadButtonMap;
+	}
+};
+
+
+FGameInputArcadeStickProcessor::FGameInputArcadeStickProcessor()
+	: IGameInputDeviceProcessor()
+{
+	FMemory::Memset(PreviousState, 0);
+	FMemory::Memset(RepeatTime, 0);
+}
+
+bool FGameInputArcadeStickProcessor::ProcessInput(const FGameInputEventParams& Params)
+{
+	bool bRes = false;
+
+	// Can't do anything for an invalid platform user or no reading.
+	if (!Params.PlatformUserId.IsValid() || !Params.Reading)
+	{
+		return bRes;
+	}
+
+	GameInputArcadeStickState StickState;
+	if (!Params.Reading->GetArcadeStickState(&StickState))
+	{
+		return bRes;
+	}
+
+	const uint32 CurrentButtonHeldMask = static_cast<uint32>(StickState.buttons);
+	uint32 LastButtonHeldMask = static_cast<uint32>(PreviousState.buttons);
+
+	EvaluateButtonStates(
+		Params,
+		CurrentButtonHeldMask,
+		OUT LastButtonHeldMask,
+		RepeatTime,
+		UE::GameInput::GetArcadeStickButtonMap(),
+		MaxSupportedButtons);
+
+	// Keep track of the button state so that we can compare it next time it is processed
+	PreviousState = StickState;
+
+	return true;
+}
+
+void FGameInputArcadeStickProcessor::ClearState(const FGameInputEventParams& Params)
+{	
+	// Can't do anything for an invalid platform user
+	if (!Params.PlatformUserId.IsValid())
+	{
+		return;
+	}
+
+	// Evaluate the buttons as if none have been pressed (i.e. the button mask is 0)
+	constexpr uint32 CurrentButtonHeldMask = 0x00;
+	uint32 LastButtonHeldMask = static_cast<uint32>(PreviousState.buttons);
+
+	EvaluateButtonStates(
+		Params,
+		CurrentButtonHeldMask,
+		OUT LastButtonHeldMask,
+		RepeatTime,
+		UE::GameInput::GetArcadeStickButtonMap(),
+		MaxSupportedButtons);
+
+	// Zero out the repeat and previous states
+	FMemory::Memset(PreviousState, 0);
+	FMemory::Memset(RepeatTime, 0);
+}
+
+GameInputKind FGameInputArcadeStickProcessor::GetSupportedReadingKind() const
+{
+	return GameInputKindArcadeStick;
+}
+
+#endif	// GAME_INPUT_SUPPORT
\ No newline at end of file
diff --git a/Engine/Plugins/Runtime/GameInput/Source/GameInputBase/Private/Processors/GameInputDeviceProcessor_Controller.cpp b/Engine/Plugins/Runtime/GameInput/Source/GameInputBase/Private/Processors/GameInputDeviceProcessor_Controller.cpp
new file mode 100644
--- /dev/null
+++ b/Engine/Plugins/Runtime/GameInput/Source/GameInputBase/Private/Processors/GameInputDeviceProcessor_Controller.cpp
@@ -0,0 +1,413 @@
+// Copyright Epic Games, Inc. All Rights Reserved.
+
+#include "Processors/GameInputDeviceProcessor_Controller.h"
+
+#include "GameInputDeveloperSettings.h"
+#include "GameInputLogging.h"
+#include "GameInputUtils.h"
+#include "InputCoreTypes.h"
+
+#if GAME_INPUT_SUPPORT
+
+/**
+* Returns a map of GameInputLabel's to their associated Unreal Engine FKey names.
+*/
+const TMap<GameInputLabel, FGamepadKeyNames::Type>& FGameInputControllerDeviceProcessor::GetGameInputButtonLabelToUnrealName()
+{
+	static const TMap<GameInputLabel, FGamepadKeyNames::Type> LabelMap
+	{
+		{ GameInputLabelUnknown, FGamepadKeyNames::Invalid },
+		{ GameInputLabelNone, FGamepadKeyNames::Invalid },
+		{ GameInputLabelXboxGuide, FGamepadKeyNames::SpecialRight },	// TODO: Check if this is correct
+		{ GameInputLabelXboxBack, FGamepadKeyNames::SpecialLeft },
+		{ GameInputLabelXboxStart, FGamepadKeyNames::SpecialRight },
+		{ GameInputLabelXboxMenu, FGamepadKeyNames::SpecialRight },
+		{ GameInputLabelXboxView, FGamepadKeyNames::SpecialLeft },
+		{ GameInputLabelXboxA, FGamepadKeyNames::FaceButtonBottom },
+		{ GameInputLabelXboxB, FGamepadKeyNames::FaceButtonRight },
+		{ GameInputLabelXboxX, FGamepadKeyNames::FaceButtonLeft },
+		{ GameInputLabelXboxY, FGamepadKeyNames::FaceButtonTop },
+		{ GameInputLabelXboxDPadUp, FGamepadKeyNames::DPadUp },
+		{ GameInputLabelXboxDPadDown, FGamepadKeyNames::DPadDown },
+		{ GameInputLabelXboxDPadLeft, FGamepadKeyNames::DPadLeft },
+		{ GameInputLabelXboxDPadRight, FGamepadKeyNames::DPadRight },
+		{ GameInputLabelXboxLeftShoulder, FGamepadKeyNames::LeftShoulder },
+		{ GameInputLabelXboxLeftTrigger, FGamepadKeyNames::LeftTriggerAnalog },
+		{ GameInputLabelXboxLeftStickButton, FGamepadKeyNames::LeftThumb },
+		{ GameInputLabelXboxRightShoulder, FGamepadKeyNames::RightShoulder },
+		{ GameInputLabelXboxRightTrigger, FGamepadKeyNames::RightTriggerAnalog },
+		{ GameInputLabelXboxRightStickButton, FGamepadKeyNames::RightThumb },
+		{ GameInputLabelXboxPaddle1, FGamepadKeyNames::Invalid },				// TODO: Do we need special additional FKey's for these paddle types?
+		{ GameInputLabelXboxPaddle2, FGamepadKeyNames::Invalid },				// Return invalid for now, but I thought the Xbox One pro controller would
+		{ GameInputLabelXboxPaddle3, FGamepadKeyNames::Invalid },				// be handled by the OS itself via virtual remapping
+		{ GameInputLabelXboxPaddle4, FGamepadKeyNames::Invalid },
+		{ GameInputLabelLetterA, EKeys::A.GetFName() },
+		{ GameInputLabelLetterB, EKeys::B.GetFName() },
+		{ GameInputLabelLetterC, EKeys::C.GetFName() },
+		{ GameInputLabelLetterD, EKeys::D.GetFName() },
+		{ GameInputLabelLetterE, EKeys::E.GetFName() },
+		{ GameInputLabelLetterF, EKeys::F.GetFName() },
+		{ GameInputLabelLetterG, EKeys::G.GetFName() },
+		{ GameInputLabelLetterH, EKeys::H.GetFName() },
+		{ GameInputLabelLetterI, EKeys::I.GetFName() },
+		{ GameInputLabelLetterJ, EKeys::J.GetFName() },
+		{ GameInputLabelLetterK, EKeys::K.GetFName() },
+		{ GameInputLabelLetterL, EKeys::L.GetFName() },
+		{ GameInputLabelLetterM, EKeys::M.GetFName() },
+		{ GameInputLabelLetterN, EKeys::N.GetFName() },
+		{ GameInputLabelLetterO, EKeys::O.GetFName() },
+		{ GameInputLabelLetterP, EKeys::P.GetFName() },
+		{ GameInputLabelLetterQ, EKeys::Q.GetFName() },
+		{ GameInputLabelLetterR, EKeys::R.GetFName() },
+		{ GameInputLabelLetterS, EKeys::S.GetFName() },
+		{ GameInputLabelLetterT, EKeys::T.GetFName() },
+		{ GameInputLabelLetterU, EKeys::U.GetFName() },
+		{ GameInputLabelLetterV, EKeys::V.GetFName() },
+		{ GameInputLabelLetterW, EKeys::W.GetFName() },
+		{ GameInputLabelLetterX, EKeys::X.GetFName() },
+		{ GameInputLabelLetterY, EKeys::Y.GetFName() },
+		{ GameInputLabelLetterZ, EKeys::Z.GetFName() },
+		{ GameInputLabelNumber0, EKeys::Zero.GetFName() },
+		{ GameInputLabelNumber1, EKeys::One.GetFName() },
+		{ GameInputLabelNumber2, EKeys::Two.GetFName() },
+		{ GameInputLabelNumber3, EKeys::Three.GetFName() },
+		{ GameInputLabelNumber4, EKeys::Four.GetFName() },
+		{ GameInputLabelNumber5, EKeys::Five.GetFName() },
+		{ GameInputLabelNumber6, EKeys::Six.GetFName() },
+		{ GameInputLabelNumber7, EKeys::Seven.GetFName() },
+		{ GameInputLabelNumber8, EKeys::Eight.GetFName() },
+		{ GameInputLabelNumber9, EKeys::Nine.GetFName() },
+		{ GameInputLabelArrowUp, EKeys::Up.GetFName() },
+		{ GameInputLabelArrowUpRight, EKeys::Up.GetFName() },
+		{ GameInputLabelArrowRight, EKeys::Right.GetFName() },
+		{ GameInputLabelArrowDownRight, EKeys::Down.GetFName() },	// TODO: We should support multiple FKey's here, like we do for switches
+		{ GameInputLabelArrowDown, EKeys::Down.GetFName() },
+		{ GameInputLabelArrowDownLLeft, EKeys::Down.GetFName() },
+		{ GameInputLabelArrowLeft, EKeys::Left.GetFName() },
+		{ GameInputLabelArrowUpLeft, FGamepadKeyNames::DPadUp },
+		{ GameInputLabelArrowUpDown, FGamepadKeyNames::DPadUp },
+		{ GameInputLabelArrowLeftRight, FGamepadKeyNames::DPadUp },
+		{ GameInputLabelArrowUpDownLeftRight, FGamepadKeyNames::DPadUp },
+		{ GameInputLabelArrowClockwise, FGamepadKeyNames::DPadUp },		// TODO: new key for this          
+		{ GameInputLabelArrowCounterClockwise, FGamepadKeyNames::DPadUp },	// TODO: new key for this   
+		{ GameInputLabelArrowReturn, EKeys::Enter.GetFName() },
+		{ GameInputLabelIconBranding, EKeys::Home.GetFName() },		// TODO: I dont think we have a UE key for this, maybe we use home?     
+		{ GameInputLabelIconHome, FGamepadKeyNames::SpecialRight },
+		{ GameInputLabelIconMenu, FGamepadKeyNames::SpecialLeft },
+		{ GameInputLabelIconCross, FGamepadKeyNames::FaceButtonBottom },
+		{ GameInputLabelIconCircle, FGamepadKeyNames::FaceButtonRight },
+		{ GameInputLabelIconSquare, FGamepadKeyNames::FaceButtonLeft },
+		{ GameInputLabelIconTriangle, FGamepadKeyNames::FaceButtonTop },
+		{ GameInputLabelIconStar, EKeys::Asterix.GetFName() },    // TODO: Star? Is this the asterix?            
+		{ GameInputLabelIconDPadUp, FGamepadKeyNames::DPadUp },
+		{ GameInputLabelIconDPadDown, FGamepadKeyNames::DPadDown },
+		{ GameInputLabelIconDPadLeft, FGamepadKeyNames::DPadLeft },
+		{ GameInputLabelIconDPadRight, FGamepadKeyNames::DPadRight },
+		{ GameInputLabelIconDialClockwise, FGamepadKeyNames::DPadUp },
+		{ GameInputLabelIconDialCounterClockwise, FGamepadKeyNames::DPadUp },
+		{ GameInputLabelIconSliderLeftRight, FGamepadKeyNames::DPadUp },
+		{ GameInputLabelIconSliderUpDown, FGamepadKeyNames::DPadUp },
+		{ GameInputLabelIconWheelUpDown, FGamepadKeyNames::DPadUp },
+		{ GameInputLabelIconPlus, EKeys::Add.GetFName() },
+		{ GameInputLabelIconMinus, EKeys::Subtract.GetFName() },
+		{ GameInputLabelIconSuspension, FGamepadKeyNames::DPadUp },
+		{ GameInputLabelHome, EKeys::Home.GetFName() },			// TODO: Do we have a gamepad key for guide?            
+		{ GameInputLabelGuide, FGamepadKeyNames::SpecialLeft },
+		{ GameInputLabelMode, FGamepadKeyNames::SpecialLeft },
+		{ GameInputLabelSelect, FGamepadKeyNames::SpecialRight },
+		{ GameInputLabelMenu, FGamepadKeyNames::SpecialRight },
+		{ GameInputLabelView, FGamepadKeyNames::SpecialLeft },
+		{ GameInputLabelBack, FGamepadKeyNames::SpecialLeft },
+		{ GameInputLabelStart, FGamepadKeyNames::SpecialRight },
+		{ GameInputLabelOptions, FGamepadKeyNames::SpecialRight },
+		{ GameInputLabelShare, FGamepadKeyNames::SpecialLeft },
+		{ GameInputLabelUp, FGamepadKeyNames::DPadUp },
+		{ GameInputLabelDown, FGamepadKeyNames::DPadDown },
+		{ GameInputLabelLeft, FGamepadKeyNames::DPadLeft },
+		{ GameInputLabelRight, FGamepadKeyNames::DPadRight },
+		{ GameInputLabelLB, FGamepadKeyNames::LeftShoulder },
+		{ GameInputLabelLT, FGamepadKeyNames::LeftTriggerAnalog },
+		{ GameInputLabelLSB, FGamepadKeyNames::LeftShoulder },
+		{ GameInputLabelL1, FGamepadKeyNames::LeftShoulder },
+		{ GameInputLabelL2, FGamepadKeyNames::LeftTriggerAnalog },
+		{ GameInputLabelL3, FGamepadKeyNames::DPadUp },
+		{ GameInputLabelRB, FGamepadKeyNames::RightShoulder },
+		{ GameInputLabelRT, FGamepadKeyNames::RightTriggerAnalog },
+		{ GameInputLabelRSB, FGamepadKeyNames::RightShoulder },
+		{ GameInputLabelR1, FGamepadKeyNames::RightShoulder },
+		{ GameInputLabelR2, FGamepadKeyNames::RightTriggerAnalog },
+		{ GameInputLabelR3, FGamepadKeyNames::DPadUp },
+	};
+
+	return LabelMap;
+}
+
+FGameInputControllerDeviceProcessor::FGameInputControllerDeviceProcessor()
+	: IGameInputDeviceProcessor()
+{
+	// Ensure that our switch repeat times array has some default values set on it by default
+	// because we access it with the [] operator when processing it. 
+	SwitchRepeatTimes.AddDefaulted(static_cast<uint32>(GameInputSwitchUpLeft) + 1);
+}
+
+bool FGameInputControllerDeviceProcessor::ProcessInput(const FGameInputEventParams& Params)
+{
+	bool bRes = false;
+
+	const FGameInputDeviceConfiguration* ControllerConfig = GetDefault<UGameInputDeveloperSettings>()->FindDeviceConfiguration(Params.GetDeviceInfo());
+
+	// Note that we use the current reading here. "ProcessInput" can be called multiple times per frame
+	// if there is more then one input reading in the stack. We want to process all button and switch states
+	// to ensure that we don't miss one
+	bRes |= ProcessControllerSwitchState(Params, ControllerConfig, Params.Reading);
+	bRes |= ProcessControllerButtonState(Params, ControllerConfig, Params.Reading);
+	++NumReadingsProcessedThisFrame;
+
+	return bRes;
+}
+
+bool FGameInputControllerDeviceProcessor::PostProcessInput(const FGameInputEventParams& Params)
+{
+	// On the last input reading for the frame, the "Current Reading" should always be null. we only care for the LastReading here
+	ensure(Params.Reading == nullptr);
+
+	bool bRes = false;
+
+	const FGameInputDeviceConfiguration* ControllerConfig = GetDefault<UGameInputDeveloperSettings>()->FindDeviceConfiguration(Params.GetDeviceInfo());
+
+	const bool bProcessedAnyButtonsThisFrame = (NumReadingsProcessedThisFrame > 0);
+
+	// This is the last reading of this frame, reset the counter to 0
+	NumReadingsProcessedThisFrame = 0;
+
+	// Note that we use "Params.PreviousReading", because the current reading will be null. 
+	if (!bProcessedAnyButtonsThisFrame)
+	{
+		bRes |= ProcessControllerSwitchState(Params, ControllerConfig, Params.PreviousReading);
+		bRes |= ProcessControllerButtonState(Params, ControllerConfig, Params.PreviousReading);
+	}
+
+	bRes |= ProcessControllerAxisState(Params, ControllerConfig, Params.PreviousReading);
+
+	return bRes;
+}
+
+void FGameInputControllerDeviceProcessor::ClearState(const FGameInputEventParams& Params)
+{
+	// Reset the switches to all be "center" positions
+	{
+		const uint32 SwitchCount = static_cast<uint32>(PreviousSwitchPositions.Num());
+
+		TArray<GameInputSwitchPosition> SwitchPositions;
+		SwitchPositions.AddUninitialized(SwitchCount);
+		// Evaluate all switch positions
+		for (uint32 i = 0; i < SwitchCount; ++i)
+		{
+			// The PreviousSwitchPositions will be updated to the current switch position by the evaluate function
+			EvaluateSwitchState(
+				Params,
+				SwitchPositions[i],
+				OUT PreviousSwitchPositions[i],
+				SwitchRepeatTimes);
+		}
+	}
+	
+
+	if (const FGameInputDeviceConfiguration* Config = GetDefault<UGameInputDeveloperSettings>()->FindDeviceConfiguration(Params.GetDeviceInfo()))
+	{
+		// Reset the button state to not being pressed
+		// Just use 0 as our button mask because we want them all to be set to 0
+		uint32 CurrentButtonHeldMask = 0;
+
+		EvaluateButtonStates(
+			Params,
+			CurrentButtonHeldMask,
+			LastButtonHeldMask,
+			RepeatTime,
+			Config->ControllerButtonMappingData,
+			MaxSupportedButtons);
+			
+		// Reset the axis state to be zero on all axis that we had in the previous state
+		for (int32 i = 0; i < PreviousControllerAxisValues.Num(); ++i)
+		{
+			if (const FGameInputControllerAxisData* AxisData = Config->ControllerAxisMappingData.Find(i))
+			{
+				OnControllerAnalog(Params, AxisData->KeyName, 0.0f, PreviousControllerAxisValues[i], AxisData->DeadZone);
+			}			
+		}
+
+		FMemory::Memset(PreviousControllerAxisValues, 0);
+	}
+}
+
+GameInputKind FGameInputControllerDeviceProcessor::GetSupportedReadingKind() const
+{
+	return GameInputKindController | GameInputKindControllerAxis | GameInputKindControllerButton;
+}
+
+bool FGameInputControllerDeviceProcessor::ProcessControllerButtonState(const FGameInputEventParams& Params, const FGameInputDeviceConfiguration* ControllerConfig, IGameInputReading* InputReading)
+{
+	// We can only process generic controllers that have their config set up
+	// but allow you to run without a config to log out button indexes to make it 
+	// easier to configure your device
+	if (ControllerConfig && !ControllerConfig->bProcessControllerButtons)
+	{
+		return false;
+	}
+
+	const uint32 ButtonCount = InputReading->GetControllerButtonCount();
+
+	TArray<bool> ButtonStates;
+	ButtonStates.AddUninitialized(ButtonCount);
+
+	const uint32 Res = InputReading->GetControllerButtonState(ButtonCount, ButtonStates.GetData());
+	if (!Res)
+	{
+		return false;
+	}
+	
+	uint32 CurrentButtonHeldMask = 0x00;
+
+	for (uint32 i = 0; i < ButtonCount; ++i)
+	{
+		const FName* KeyName = ControllerConfig ? ControllerConfig->ControllerButtonMappingData.Find(1 << i) : nullptr;
+
+		CurrentButtonHeldMask |= (ButtonStates[i] << i);
+
+		if (ButtonStates[i])
+		{
+			UE_LOG(LogGameInput, Verbose, TEXT("[ProcessControllerButtonState] Device ID: %d   Button count: %d    index: %d  State: %d   %s"), 
+				Params.InputDeviceId.GetId(),
+				ButtonCount, 
+				i,
+				ButtonStates[i], 
+				KeyName ? *(KeyName->ToString()) : TEXT("NONE"));
+		}
+	}
+
+	if (ControllerConfig)
+	{
+		EvaluateButtonStates(
+			Params,
+			CurrentButtonHeldMask,
+			OUT LastButtonHeldMask,
+			RepeatTime,
+			ControllerConfig->ControllerButtonMappingData,
+			ButtonCount
+		);
+	}	
+
+	// Note: Ideally we would try and use the GameInput label system, 
+	// but it is currently unfinished within the GameInput API itself so we can't.
+	// We have provided this device specific config driven option instead.
+
+	return true;
+}
+
+bool FGameInputControllerDeviceProcessor::ProcessControllerAxisState(const FGameInputEventParams& Params, const FGameInputDeviceConfiguration* Config, IGameInputReading* InputReading)
+{
+	if (Config && !Config->bProcessControllerAxis)
+	{
+		return false;
+	}
+
+	const uint32 AxisCount = InputReading->GetControllerAxisCount();
+	TArray<float> AxisValues;
+	AxisValues.AddZeroed(AxisCount);
+	if (!InputReading->GetControllerAxisState(AxisCount, AxisValues.GetData()))
+	{
+		return false;
+	}
+
+	// Make sure that there is previous controller data initialized if necessary
+	while (PreviousControllerAxisValues.Num() < static_cast<int32>(AxisCount))
+	{
+		PreviousControllerAxisValues.AddZeroed();
+	}
+
+	for (uint32 i = 0; i < AxisCount; ++i)
+	{
+		float CurrentValue = AxisValues[i];
+		const float PreviousValue = PreviousControllerAxisValues[i];
+
+		if (Config)
+		{
+			if (const FGameInputControllerAxisData* AxisData = Config->ControllerAxisMappingData.Find(i))
+			{
+				if (AxisData->KeyName.IsValid() && AxisData->KeyName != NAME_None)
+				{
+					if (AxisData->bIsPackedPositveAndNegative)
+					{
+						// Maps the value to be -1.0 to +1.0
+						CurrentValue = (CurrentValue * 2.f) - 1.f;
+					}
+
+					CurrentValue *= AxisData->Scalar;
+					
+					OnControllerAnalog(Params, AxisData->KeyName, CurrentValue, PreviousValue, AxisData->DeadZone);
+
+					// Store this value for the next frame to compare to
+					PreviousControllerAxisValues[i] = CurrentValue;
+				}
+				else
+				{
+					UE_LOG(LogGameInput, VeryVerbose, TEXT("[ProcessControllerAxisState] (Device %s) Invalid key name configured for controller axis '%d': %.3f"), *UE::GameInput::LexToString(Params.Device), i, CurrentValue);
+				}
+			}
+			else
+			{
+				// TODO: Here is where could send a "Generic USB Axis X" key here which could allow for us to support many more devices via a key rebind screen
+				UE_LOG(LogGameInput, VeryVerbose, TEXT("[ProcessControllerAxisState] (Device %s) Controller axis '%d' has value: %.3f"), *UE::GameInput::LexToString(Params.Device), i, CurrentValue);
+			}
+		}		
+		// You are receiving analog values from an axis that you might not know about, log it here
+		// in case you are trying to set something up
+		else
+		{
+			UE_LOG(LogGameInput, VeryVerbose, TEXT("[ProcessControllerAxisState] (Device %s) Receiving input from an unconfigured controller axis '%d': %.3f"), *UE::GameInput::LexToString(Params.Device), i, CurrentValue);
+		}
+	}
+
+	return false;
+}
+
+bool FGameInputControllerDeviceProcessor::ProcessControllerSwitchState(const FGameInputEventParams& Params, const FGameInputDeviceConfiguration* Config, IGameInputReading* InputReading)
+{
+	if (Config && !Config->bProcessControllerSwitchState)
+	{
+		return false;
+	}
+
+	const uint32 SwitchCount = InputReading->GetControllerSwitchCount();
+	TArray<GameInputSwitchPosition> SwitchPositions;
+	SwitchPositions.AddUninitialized(SwitchCount);
+
+	// Make sure that we have some previous state initialized if we can
+	if (static_cast<uint32>(PreviousSwitchPositions.Num()) < SwitchCount)
+	{
+		PreviousSwitchPositions.AddDefaulted(SwitchCount);
+	}
+
+	uint32 Res = InputReading->GetControllerSwitchState(SwitchCount, SwitchPositions.GetData());
+
+	if (!Res)
+	{
+		return false;
+	}
+
+	// Evaluate all switch positions
+	for (uint32 i = 0; i < SwitchCount; ++i)
+	{
+		// The PreviousSwitchPositions will be updated to the current switch position by the evaluate function
+		EvaluateSwitchState(
+			Params,
+			SwitchPositions[i],
+			OUT PreviousSwitchPositions[i],
+			SwitchRepeatTimes);
+	}
+
+	return true;
+}
+
+#endif	// #if GAME_INPUT_SUPPORT
+
diff --git a/Engine/Plugins/Runtime/GameInput/Source/GameInputBase/Private/Processors/GameInputDeviceProcessor_FlightStick.cpp b/Engine/Plugins/Runtime/GameInput/Source/GameInputBase/Private/Processors/GameInputDeviceProcessor_FlightStick.cpp
new file mode 100644
--- /dev/null
+++ b/Engine/Plugins/Runtime/GameInput/Source/GameInputBase/Private/Processors/GameInputDeviceProcessor_FlightStick.cpp
@@ -0,0 +1,195 @@
+// Copyright Epic Games, Inc. All Rights Reserved.
+
+#include "Processors/GameInputDeviceProcessor_FlightStick.h"
+
+#include "GameInputDeveloperSettings.h"
+#include "GameInputKeyTypes.h"
+
+#if GAME_INPUT_SUPPORT
+
+namespace UE::GameInput
+{
+	/** A map of uint32 GameInput button bitmask flags to the associated Unreal Engine FKey name. */
+	static const TMap<uint32, FGamepadKeyNames::Type>& GetFlightStickButtonMap()
+	{
+		static const TMap<uint32, FGamepadKeyNames::Type> GamepadButtonMap
+		{
+			// Generic gamepad buttons
+			{ static_cast<uint32>(GameInputFlightStickButtons::GameInputFlightStickNone), FGamepadKeyNames::Invalid},
+			{ static_cast<uint32>(GameInputFlightStickButtons::GameInputFlightStickMenu), FGamepadKeyNames::Invalid},
+			{ static_cast<uint32>(GameInputFlightStickButtons::GameInputFlightStickView), FGamepadKeyNames::Invalid},
+			{ static_cast<uint32>(GameInputFlightStickButtons::GameInputFlightStickFirePrimary), FGamepadKeyNames::Invalid},
+			{ static_cast<uint32>(GameInputFlightStickButtons::GameInputFlightStickFireSecondary), FGamepadKeyNames::Invalid},
+		};
+		return GamepadButtonMap;
+	}
+	
+	static bool HasDifferentFlightStickAnalogs(const GameInputFlightStickState& CurrentState, const GameInputFlightStickState& PreviousState)
+	{
+		// If anything differs from the previous reading, then it has different inputs
+		return 
+			CurrentState.pitch != PreviousState.pitch ||
+			CurrentState.roll != PreviousState.roll ||
+			CurrentState.throttle != PreviousState.throttle ||
+			CurrentState.yaw != PreviousState.yaw;
+	}
+
+	inline bool IsEmptyFlightStickReading(const GameInputFlightStickState& State, const UGameInputPlatformSettings& PlatformSettings)
+	{
+		return
+			State.buttons == 0 &&
+			State.hatSwitch == 0 &&	
+			FMath::Abs(State.pitch) <= PlatformSettings.FlightStickPitchDeadzone &&
+			FMath::Abs(State.roll) <= PlatformSettings.FlightStickRollDeadzone &&
+			FMath::Abs(State.throttle) <= PlatformSettings.FlightStickThrottleDeadzone &&
+			FMath::Abs(State.yaw) <= PlatformSettings.FlightStickYawDeadzone; 
+	};
+}
+
+FGameInputFlightStickProcessor::FGameInputFlightStickProcessor()
+	: IGameInputDeviceProcessor()
+{
+	FMemory::Memset(PreviousState, 0);
+	FMemory::Memset(RepeatTime, 0);
+	
+	SwitchRepeatTimes.AddDefaulted(static_cast<uint32>(GameInputSwitchUpLeft) + 1);
+}
+
+bool FGameInputFlightStickProcessor::ProcessInput(const FGameInputEventParams& Params)
+{
+	bool bRes = false;
+	
+	// Can't do anything for an invalid platform user
+	if (!Params.PlatformUserId.IsValid() || !Params.Reading)
+	{
+		return bRes;
+	}
+
+	GameInputFlightStickState FlightStickState = {};
+	if (!Params.Reading->GetFlightStickState(&FlightStickState))
+	{
+		return bRes;
+	}
+
+	// We only want to process the buttons here, as it might get called multiple times per frame.
+	bRes |= ProcessFlightStickButtons(Params, FlightStickState);
+	
+	++NumReadingsProcessedThisFrame;
+	
+	return bRes;
+}
+
+bool FGameInputFlightStickProcessor::PostProcessInput(const FGameInputEventParams& Params)
+{
+	bool bRes = false;
+	
+	// Check if we have already processed buttons this frame. If we haven't we want to do it
+	const bool bHasProcessedAnyButtonsThisFrame = NumReadingsProcessedThisFrame > 0;	
+	NumReadingsProcessedThisFrame = 0;
+	
+	// Can't do anything for an invalid platform user
+	if (!Params.PlatformUserId.IsValid() || !Params.PreviousReading)
+	{
+		return bRes;
+	}
+
+	GameInputFlightStickState FlightStickState = {};
+	if (!Params.PreviousReading->GetFlightStickState(&FlightStickState))
+	{
+		return bRes;
+	}
+	
+	if (!bHasProcessedAnyButtonsThisFrame)
+	{
+		bRes |= ProcessFlightStickButtons(Params, FlightStickState);
+	}
+	
+	bRes |= ProcessFlightStickAnalog(Params, FlightStickState);
+	
+	return bRes;
+}
+
+void FGameInputFlightStickProcessor::ClearState(const FGameInputEventParams& Params)
+{
+	// Can't do anything for an invalid platform user
+	if (!Params.PlatformUserId.IsValid())
+	{
+		return;
+	}
+	
+	// Process input as if nothing is down (zero values for everything)
+	GameInputFlightStickState ZeroState = {};
+	ProcessFlightStickAnalog(Params, ZeroState);
+	ProcessFlightStickButtons(Params, ZeroState);
+	
+	// Zero out the repeat and previous states info to zero
+	FMemory::Memset(PreviousState, 0);
+	FMemory::Memset(RepeatTime, 0);
+}
+
+GameInputKind FGameInputFlightStickProcessor::GetSupportedReadingKind() const
+{
+	return GameInputKindFlightStick;
+}
+
+bool FGameInputFlightStickProcessor::ProcessFlightStickButtons(const FGameInputEventParams& Params, GameInputFlightStickState& State)
+{
+	const UGameInputPlatformSettings* PlatformSettings = UGameInputPlatformSettings::Get();
+
+	const bool bIsEmptyReading = UE::GameInput::IsEmptyFlightStickReading(State, *PlatformSettings);
+	const bool bWasEmptyReading = UE::GameInput::IsEmptyFlightStickReading(PreviousState, *PlatformSettings);
+	const bool bHasDifferentAnalogInput = UE::GameInput::HasDifferentFlightStickAnalogs(State, PreviousState);
+	
+	if (bIsEmptyReading && bWasEmptyReading && !bHasDifferentAnalogInput && PreviousState.buttons == 0)
+	{
+		return false;
+	}
+	
+	const uint32 CurrentButtonHeldMask = static_cast<uint32>(State.buttons);
+	uint32 LastButtonHeldMask = static_cast<uint32>(PreviousState.buttons);
+
+	EvaluateButtonStates(
+		Params,
+		CurrentButtonHeldMask,
+		OUT LastButtonHeldMask,
+		RepeatTime,
+		UE::GameInput::GetFlightStickButtonMap(),
+		MaxSupportedButtons);
+
+	// Update the previous state here
+	PreviousState.buttons = State.buttons;
+
+	// Update the hat switch
+	EvaluateSwitchState(Params, State.hatSwitch, PreviousState.hatSwitch, SwitchRepeatTimes);
+	
+	return true;
+}
+
+bool FGameInputFlightStickProcessor::ProcessFlightStickAnalog(const FGameInputEventParams& Params, GameInputFlightStickState& State)
+{
+	const UGameInputPlatformSettings* Settings = UGameInputPlatformSettings::Get();
+	check(Settings);
+	
+	// If the analog values haven't changed, don't bother sending any events for them
+	const bool bIsEmptyReading = UE::GameInput::IsEmptyFlightStickReading(State, *Settings);
+	const bool bWasEmptyReading = UE::GameInput::IsEmptyFlightStickReading(PreviousState, *Settings);
+	const bool bHasDifferentAnalogInput = UE::GameInput::HasDifferentFlightStickAnalogs(State, PreviousState);
+	if (bIsEmptyReading && bWasEmptyReading && !bHasDifferentAnalogInput)
+	{
+		return false;
+	}
+	
+	OnControllerAnalog(Params, FGameInputKeys::FlightStick_Pitch.GetFName(), State.pitch, PreviousState.pitch, Settings->FlightStickPitchDeadzone);
+	OnControllerAnalog(Params, FGameInputKeys::FlightStick_Roll.GetFName(), State.roll, PreviousState.roll, Settings->FlightStickRollDeadzone);
+	OnControllerAnalog(Params, FGameInputKeys::FlightStick_Throttle.GetFName(), State.throttle, PreviousState.throttle, Settings->FlightStickThrottleDeadzone);
+	OnControllerAnalog(Params, FGameInputKeys::FlightStick_Yaw.GetFName(), State.yaw, PreviousState.yaw, Settings->FlightStickYawDeadzone);
+
+	PreviousState.pitch = State.pitch;
+	PreviousState.roll = State.roll;
+	PreviousState.throttle = State.throttle;
+	PreviousState.yaw = State.yaw;
+	
+	return true;
+}
+
+#endif
\ No newline at end of file
diff --git a/Engine/Plugins/Runtime/GameInput/Source/GameInputBase/Private/Processors/GameInputDeviceProcessor_Gamepad.cpp b/Engine/Plugins/Runtime/GameInput/Source/GameInputBase/Private/Processors/GameInputDeviceProcessor_Gamepad.cpp
new file mode 100644
--- /dev/null
+++ b/Engine/Plugins/Runtime/GameInput/Source/GameInputBase/Private/Processors/GameInputDeviceProcessor_Gamepad.cpp
@@ -0,0 +1,294 @@
+// Copyright Epic Games, Inc. All Rights Reserved.
+
+#include "Processors/GameInputDeviceProcessor_Gamepad.h"
+
+#include "GameInputDeveloperSettings.h"
+#include "GameInputUtils.h"
+#include "GenericPlatform/GenericApplicationMessageHandler.h"
+
+#if GAME_INPUT_SUPPORT
+
+namespace UE::GameInput
+{
+	inline bool IsEmptyGamepadReading(const GameInputGamepadState& GamepadState)
+	{
+		return GamepadState.buttons == 0 &&
+			GamepadState.leftTrigger <= GamepadTriggerDeadzone && GamepadState.rightTrigger <= GamepadTriggerDeadzone &&
+			FMath::Abs(GamepadState.leftThumbstickX) <= GamepadLeftStickDeadzone && FMath::Abs(GamepadState.leftThumbstickY) <= GamepadLeftStickDeadzone &&
+			FMath::Abs(GamepadState.rightThumbstickX) <= GamepadRightStickDeadzone && FMath::Abs(GamepadState.rightThumbstickY) <= GamepadRightStickDeadzone;
+	};
+
+	inline bool HasDifferentAnalogInput(const GameInputGamepadState& CurrentGamepadState, const GameInputGamepadState& PreviousGamepadState)
+	{
+		return CurrentGamepadState.leftTrigger != PreviousGamepadState.leftTrigger || CurrentGamepadState.rightTrigger != PreviousGamepadState.rightTrigger ||
+			CurrentGamepadState.leftThumbstickX != PreviousGamepadState.leftThumbstickX || CurrentGamepadState.leftThumbstickY != PreviousGamepadState.leftThumbstickY ||
+			CurrentGamepadState.rightThumbstickX != PreviousGamepadState.rightThumbstickX || CurrentGamepadState.rightThumbstickY != PreviousGamepadState.rightThumbstickY;
+	};
+
+	/**
+	 * Returns a map of GameInputGamepadButtons to their associated Unreal Engine FKey names.
+	 */
+	static const TMap<uint32, FGamepadKeyNames::Type>& GetGamepadButtonMap()
+	{
+		static const TMap<uint32, FGamepadKeyNames::Type> GamepadButtonMap
+		{
+			{ static_cast<uint32>(GameInputGamepadButtons::GameInputGamepadA), FGamepadKeyNames::FaceButtonBottom },
+			{ static_cast<uint32>(GameInputGamepadButtons::GameInputGamepadB), FGamepadKeyNames::FaceButtonRight },
+			{ static_cast<uint32>(GameInputGamepadButtons::GameInputGamepadX), FGamepadKeyNames::FaceButtonLeft },
+			{ static_cast<uint32>(GameInputGamepadButtons::GameInputGamepadY), FGamepadKeyNames::FaceButtonTop },
+			{ static_cast<uint32>(GameInputGamepadButtons::GameInputGamepadLeftShoulder), FGamepadKeyNames::LeftShoulder },
+			{ static_cast<uint32>(GameInputGamepadButtons::GameInputGamepadRightShoulder), FGamepadKeyNames::RightShoulder },
+			{ static_cast<uint32>(GameInputGamepadButtons::GameInputGamepadMenu), FGamepadKeyNames::SpecialRight },
+			{ static_cast<uint32>(GameInputGamepadButtons::GameInputGamepadView), FGamepadKeyNames::SpecialLeft },
+			{ static_cast<uint32>(GameInputGamepadButtons::GameInputGamepadDPadUp), FGamepadKeyNames::DPadUp },
+			{ static_cast<uint32>(GameInputGamepadButtons::GameInputGamepadDPadDown), FGamepadKeyNames::DPadDown },
+			{ static_cast<uint32>(GameInputGamepadButtons::GameInputGamepadDPadLeft), FGamepadKeyNames::DPadLeft },
+			{ static_cast<uint32>(GameInputGamepadButtons::GameInputGamepadDPadRight), FGamepadKeyNames::DPadRight },
+			{ static_cast<uint32>(GameInputGamepadButtons::GameInputGamepadLeftThumbstick), FGamepadKeyNames::LeftThumb },
+			{ static_cast<uint32>(GameInputGamepadButtons::GameInputGamepadRightThumbstick), FGamepadKeyNames::RightThumb },
+
+			{ static_cast<uint32>(EGameInputAuxButton::GamepadButtonAux_LeftTrigger), FGamepadKeyNames::LeftTriggerThreshold },
+			{ static_cast<uint32>(EGameInputAuxButton::GamepadButtonAux_RightTrigger), FGamepadKeyNames::RightTriggerThreshold },
+			{ static_cast<uint32>(EGameInputAuxButton::GamepadButtonAux_LeftStickUp), FGamepadKeyNames::LeftStickUp },
+			{ static_cast<uint32>(EGameInputAuxButton::GamepadButtonAux_LeftStickDown), FGamepadKeyNames::LeftStickDown },
+			{ static_cast<uint32>(EGameInputAuxButton::GamepadButtonAux_LeftStickLeft), FGamepadKeyNames::LeftStickLeft },
+			{ static_cast<uint32>(EGameInputAuxButton::GamepadButtonAux_LeftStickRight), FGamepadKeyNames::LeftStickRight },
+			{ static_cast<uint32>(EGameInputAuxButton::GamepadButtonAux_RightStickUp), FGamepadKeyNames::RightStickUp },
+			{ static_cast<uint32>(EGameInputAuxButton::GamepadButtonAux_RightStickDown), FGamepadKeyNames::RightStickDown },
+			{ static_cast<uint32>(EGameInputAuxButton::GamepadButtonAux_RightStickLeft), FGamepadKeyNames::RightStickLeft },
+			{ static_cast<uint32>(EGameInputAuxButton::GamepadButtonAux_RightStickRight), FGamepadKeyNames::RightStickRight }
+		};
+		return GamepadButtonMap;
+	}
+}
+
+bool FGameInputGamepadDeviceProcessor::ProcessInput(const FGameInputEventParams& Params)
+{
+	// check if the reading had gamepad info
+	GameInputGamepadState GamepadState;
+	if (!Params.Reading->GetGamepadState(&GamepadState))
+	{
+		return false;
+	}
+
+	bool bRes = false;	
+
+	// We want to process gamepad BUTTON states for every game input reading.
+	bRes |= ProcessGamepadButtonState(Params, GamepadState);
+	++NumReadingsProcessedThisFrame;
+
+	return bRes;
+}
+
+bool FGameInputGamepadDeviceProcessor::PostProcessInput(const FGameInputEventParams& Params)
+{	
+	// On the last input reading for the frame, the "Current Reading" should always be null. we only care for the LastReading here
+	ensure (Params.Reading == nullptr);
+
+	const bool bProcessedAnyGamepadButtonsThisFrame = (NumReadingsProcessedThisFrame > 0);
+
+	// This is the last reading of this frame, reset the counter to 0
+	NumReadingsProcessedThisFrame = 0;
+
+	if (!Params.PreviousReading)
+	{
+		return false;
+	}
+
+	// Get the gamepad state of the *Last Reading* of the frame here
+	GameInputGamepadState GamepadState;
+	if (!Params.PreviousReading->GetGamepadState(&GamepadState))
+	{
+		return false;
+	}
+
+	bool bRes = false;
+	
+	// We only want to process gamepad ANALOG inputs on the last reading because we only care about
+	// the most recent analog stick input value. If there are multiple readings and we processed analog
+	// for every one, then the values would accumulate and stack, giving us incorrect data in the message handler.
+	bRes |= ProcessGamepadAnalogState(Params, GamepadState);
+	
+	// If there were no gamepad events this frame, send button events using the previous frame's reading for button repeats.
+	// This is necessary because GetNextReading won't return any reading if the state is unchanged
+	if (!bProcessedAnyGamepadButtonsThisFrame)
+	{
+		bRes |= ProcessGamepadButtonState(Params, GamepadState);
+	}
+
+	return bRes;
+}
+
+bool FGameInputGamepadDeviceProcessor::ProcessGamepadAnalogState(const FGameInputEventParams& Params, GameInputGamepadState& GamepadState)
+{
+	if (!Params.PlatformUserId.IsValid())
+	{
+		return false;
+	}
+
+	// ignore this input if the reading has remained empty from last time
+	const bool bIsEmptyReading = UE::GameInput::IsEmptyGamepadReading(GamepadState);
+	const bool bWasEmptyReading = UE::GameInput::IsEmptyGamepadReading(PreviousState);
+	const bool bHasDifferentAnalogInput = UE::GameInput::HasDifferentAnalogInput(GamepadState, PreviousState);
+	if (bIsEmptyReading && bWasEmptyReading && !bHasDifferentAnalogInput)
+	{
+		return false;
+	}
+
+	OnControllerAnalog(Params, FGamepadKeyNames::LeftAnalogX, GamepadState.leftThumbstickX, PreviousState.leftThumbstickX, UE::GameInput::GamepadLeftStickDeadzone);
+	OnControllerAnalog(Params, FGamepadKeyNames::LeftAnalogY, GamepadState.leftThumbstickY, PreviousState.leftThumbstickY, UE::GameInput::GamepadLeftStickDeadzone);
+	OnControllerAnalog(Params, FGamepadKeyNames::RightAnalogX, GamepadState.rightThumbstickX, PreviousState.rightThumbstickX, UE::GameInput::GamepadRightStickDeadzone);
+	OnControllerAnalog(Params, FGamepadKeyNames::RightAnalogY, GamepadState.rightThumbstickY, PreviousState.rightThumbstickY, UE::GameInput::GamepadRightStickDeadzone);
+	OnControllerAnalog(Params, FGamepadKeyNames::LeftTriggerAnalog, GamepadState.leftTrigger, PreviousState.leftTrigger, UE::GameInput::GamepadTriggerDeadzone);
+	OnControllerAnalog(Params, FGamepadKeyNames::RightTriggerAnalog, GamepadState.rightTrigger, PreviousState.rightTrigger, UE::GameInput::GamepadTriggerDeadzone);
+	
+	// map analog triggers to digital input.
+	uint32 AnalogButtonMask = 0u;
+	if (GamepadState.leftTrigger > UE::GameInput::GamepadTriggerDeadzone)
+	{
+		AnalogButtonMask |= UE::GameInput::EGameInputAuxButton::GamepadButtonAux_LeftTrigger;
+	}
+
+	if (GamepadState.rightTrigger > UE::GameInput::GamepadTriggerDeadzone)
+	{
+		AnalogButtonMask |= UE::GameInput::EGameInputAuxButton::GamepadButtonAux_RightTrigger;
+	}
+
+	// map left/right stick digital inputs to (top 8 bits)
+	if (GamepadState.leftThumbstickY > UE::GameInput::GamepadLeftStickDeadzone)
+	{
+		AnalogButtonMask |= UE::GameInput::EGameInputAuxButton::GamepadButtonAux_LeftStickUp;
+	}
+	else if (GamepadState.leftThumbstickY < -UE::GameInput::GamepadLeftStickDeadzone)
+	{
+		AnalogButtonMask |= UE::GameInput::EGameInputAuxButton::GamepadButtonAux_LeftStickDown;
+	}
+	if (GamepadState.leftThumbstickX > UE::GameInput::GamepadLeftStickDeadzone)
+	{
+		AnalogButtonMask |= UE::GameInput::EGameInputAuxButton::GamepadButtonAux_LeftStickRight;
+	}
+	else if (GamepadState.leftThumbstickX < -UE::GameInput::GamepadLeftStickDeadzone)
+	{
+		AnalogButtonMask |= UE::GameInput::EGameInputAuxButton::GamepadButtonAux_LeftStickLeft;
+	}
+
+	if (GamepadState.rightThumbstickY > UE::GameInput::GamepadRightStickDeadzone)
+	{
+		AnalogButtonMask |= UE::GameInput::EGameInputAuxButton::GamepadButtonAux_RightStickUp;
+	}
+	else if (GamepadState.rightThumbstickY < -UE::GameInput::GamepadRightStickDeadzone)
+	{
+		AnalogButtonMask |= UE::GameInput::EGameInputAuxButton::GamepadButtonAux_RightStickDown;
+	}
+	if (GamepadState.rightThumbstickX > UE::GameInput::GamepadRightStickDeadzone)
+	{
+		AnalogButtonMask |= UE::GameInput::EGameInputAuxButton::GamepadButtonAux_RightStickRight;
+	}
+	else if (GamepadState.rightThumbstickX < -UE::GameInput::GamepadRightStickDeadzone)
+	{
+		AnalogButtonMask |= UE::GameInput::EGameInputAuxButton::GamepadButtonAux_RightStickLeft;
+	}
+
+	// Evaluate the analog button mask
+	EvaluateButtonStates(
+		Params,
+		AnalogButtonMask,
+		OUT LastAnalogButtonHeldMask,
+		RepeatTime,
+		UE::GameInput::GetGamepadButtonMap(),
+		MaxSupportedButtons);
+	
+	// update saved analog state
+	PreviousState.leftTrigger = GamepadState.leftTrigger;
+	PreviousState.rightTrigger = GamepadState.rightTrigger;
+	PreviousState.leftThumbstickX = GamepadState.leftThumbstickX;
+	PreviousState.leftThumbstickY = GamepadState.leftThumbstickY;
+	PreviousState.rightThumbstickX = GamepadState.rightThumbstickX;
+	PreviousState.rightThumbstickY = GamepadState.rightThumbstickY;
+		
+	return true;
+}
+
+bool FGameInputGamepadDeviceProcessor::ProcessGamepadButtonState(const FGameInputEventParams& Params, GameInputGamepadState& GamepadState)
+{
+	if (!Params.PlatformUserId.IsValid())
+	{
+		return false;
+	}
+
+	// ignore this input if the reading has remained empty from last time, or we still have outstanding held buttons 
+	// (these held buttons are likely to occur with mapped analog triggers when there have been multiple input events this frame)
+	const bool bIsEmptyReading = UE::GameInput::IsEmptyGamepadReading(GamepadState);
+	const bool bWasEmptyReading = UE::GameInput::IsEmptyGamepadReading(PreviousState);
+	const bool bHasDifferentAnalogInput = UE::GameInput::HasDifferentAnalogInput(GamepadState, PreviousState);
+	if (bIsEmptyReading && bWasEmptyReading && !bHasDifferentAnalogInput && LastButtonHeldMask == 0)
+	{
+		return false;
+	}
+
+	// map buttons (low 15 bits)
+	uint32 CurrentButtonHeldMask = (static_cast<uint32>(GamepadState.buttons) & UE::GameInput::GamingInputButtonMask);
+
+	EvaluateButtonStates(
+		Params,
+		CurrentButtonHeldMask,
+		OUT LastButtonHeldMask,
+		RepeatTime,
+		UE::GameInput::GetGamepadButtonMap(),
+		MaxSupportedButtons);
+
+	// Keep track of the current BUTTON state. We don't want to update the entire PreviousState struct here
+	// because buttons may be evaluated more then analog inputs per-frame
+	PreviousState.buttons = GamepadState.buttons;
+
+	return true;
+}
+
+void FGameInputGamepadDeviceProcessor::ClearState(const FGameInputEventParams& Params)
+{
+	// We need a valid input device id when sending messages, otherwise slate will hit a check
+	// and attempt to create some new slate user with an invalid index.
+	if (!Params.PlatformUserId.IsValid() || !Params.InputDeviceId.IsValid())
+	{
+		return;
+	}
+
+	// Reset Axis values
+	OnControllerAnalog(Params, FGamepadKeyNames::LeftAnalogX, 0.0f, PreviousState.leftThumbstickX, UE::GameInput::GamepadLeftStickDeadzone);
+	OnControllerAnalog(Params, FGamepadKeyNames::LeftAnalogY, 0.0f, PreviousState.leftThumbstickY, UE::GameInput::GamepadLeftStickDeadzone);
+	OnControllerAnalog(Params, FGamepadKeyNames::RightAnalogX, 0.0f, PreviousState.rightThumbstickX, UE::GameInput::GamepadRightStickDeadzone);
+	OnControllerAnalog(Params, FGamepadKeyNames::RightAnalogY, 0.0f, PreviousState.rightThumbstickY, UE::GameInput::GamepadRightStickDeadzone);
+	OnControllerAnalog(Params, FGamepadKeyNames::LeftTriggerAnalog, 0.0f, PreviousState.leftTrigger, UE::GameInput::GamepadTriggerDeadzone);
+	OnControllerAnalog(Params, FGamepadKeyNames::RightTriggerAnalog, 0.0f, PreviousState.rightTrigger, UE::GameInput::GamepadTriggerDeadzone);
+
+	// Reset button values
+
+	// Just use 0 as our button mask because we want them all to be set to 0
+	uint32 CurrentButtonHeldMask = 0;
+
+	EvaluateButtonStates(
+		Params,
+		CurrentButtonHeldMask,
+		LastButtonHeldMask,
+		RepeatTime,
+		UE::GameInput::GetGamepadButtonMap(),
+		MaxSupportedButtons);
+
+	// Clear the analog button mask
+	LastAnalogButtonHeldMask = 0u;
+
+	for (uint32 i = 0; i < MaxSupportedButtons; i++)
+	{
+		RepeatTime[i] = 0.0;
+	}
+
+	// clear previous gamepad state
+	FMemory::Memset(PreviousState, 0);
+}
+
+GameInputKind FGameInputGamepadDeviceProcessor::GetSupportedReadingKind() const
+{
+	return GameInputKindGamepad;
+}
+
+#endif	// #if GAME_INPUT_SUPPORT
\ No newline at end of file
diff --git a/Engine/Plugins/Runtime/GameInput/Source/GameInputBase/Private/Processors/GameInputDeviceProcessor_KeyboardMouse.cpp b/Engine/Plugins/Runtime/GameInput/Source/GameInputBase/Private/Processors/GameInputDeviceProcessor_KeyboardMouse.cpp
new file mode 100644
--- /dev/null
+++ b/Engine/Plugins/Runtime/GameInput/Source/GameInputBase/Private/Processors/GameInputDeviceProcessor_KeyboardMouse.cpp
@@ -0,0 +1,324 @@
+// Copyright Epic Games, Inc. All Rights Reserved.
+
+#include "Processors/GameInputDeviceProcessor_KeyboardMouse.h"
+
+#include "Framework/Application/SlateApplication.h"
+#include "GameInputDeveloperSettings.h"
+#include "GameInputLogging.h"
+#include "GameInputUtils.h"
+#include "GenericPlatform/CursorUtils.h"
+#include "HAL/ConsoleManager.h"
+#include "HAL/PlatformTime.h"
+
+#if GAME_INPUT_SUPPORT
+
+static float GKeyboardRepeatInitialDelay = 0.25f;
+static FAutoConsoleVariableRef CVarKeyboardRepeatInitialDelay(
+	TEXT("Input.KeyboardRepeatInitialDelay"),
+	GKeyboardRepeatInitialDelay,
+	TEXT("Time in seconds before a key repeat starts"));
+
+static float GKeyboardRepeatDelay = 0.05f;
+static FAutoConsoleVariableRef CVarKeyboardRepeatDelay(
+	TEXT("Input.KeyboardRepeatDelay"),
+	GKeyboardRepeatDelay,
+	TEXT("Time in seconds between each subsequent key repeat"));
+
+bool FGameInputKeyboardDeviceProcessor::ProcessInput(const FGameInputEventParams& Params)
+{	
+	TSet<uint8> CurrentPressedKeys;
+
+	int32 KeyCount = Params.Reading->GetKeyCount();
+	if (KeyCount > 0)
+	{
+		// read the key state
+		TArray<GameInputKeyState> KeyStates;
+		KeyStates.AddUninitialized(KeyCount);
+		int32 ReadCount = Params.Reading->GetKeyState(KeyCount, KeyStates.GetData());
+
+		// build a set of the pressed keycodes
+		for (GameInputKeyState KeyState : KeyStates)
+		{
+			CurrentPressedKeys.Add(KeyState.virtualKey);
+		}		
+	}
+
+	UpdateUnifiedKeyboardState(Params, CurrentPressedKeys);
+	return true;
+}
+
+void FGameInputKeyboardDeviceProcessor::UpdateUnifiedKeyboardState(const FGameInputEventParams& Params, TSet<uint8>& CurrentPressedKeys)
+{
+	// process unified pressed keys
+	double CurrentTime = FPlatformTime::Seconds();
+	for (uint8 KeyCode : CurrentPressedKeys)
+	{
+		bool bIsRepeat = LastPressedKeys.Contains(KeyCode);
+		if (!bIsRepeat)
+		{
+			KeyRepeatTime.Add(KeyCode, CurrentTime + GKeyboardRepeatInitialDelay);
+			Params.MessageHandler->OnKeyDown(KeyCode, 0, false);
+			UE_LOG(LogGameInput, Verbose, TEXT("Key Press 0x%X"), KeyCode);
+
+			if (KeyCode == VK_CAPITAL)
+			{
+				SetSimulatedCapsLock(!bSimulatedCapsLock);
+				UE_LOG(LogGameInput, Verbose, TEXT("Simulated caps lock is %s"), bSimulatedCapsLock ? TEXT("ON") : TEXT("OFF"));
+			}
+		}
+		else if (CurrentTime > KeyRepeatTime[KeyCode])
+		{
+			KeyRepeatTime.Add(KeyCode, CurrentTime + GKeyboardRepeatDelay);
+			Params.MessageHandler->OnKeyDown(KeyCode, 0, true);
+			UE_LOG(LogGameInput, Verbose, TEXT("Key Press 0x%X (repeat)"), KeyCode);
+		}
+	}
+
+	// process any released keys
+	for (uint8 KeyCode : LastPressedKeys)
+	{
+		if (!CurrentPressedKeys.Contains(KeyCode))
+		{
+			KeyRepeatTime.Remove(KeyCode);
+			Params.MessageHandler->OnKeyUp(KeyCode, 0, false);
+			UE_LOG(LogGameInput, Verbose, TEXT("Key Release 0x%X"), KeyCode);
+		}
+	}
+
+	// update saved state
+	LastPressedKeys = CurrentPressedKeys;
+}
+
+bool FGameInputKeyboardDeviceProcessor::IsKeyPressed(uint8 VirtualKeyCode) const
+{
+	return LastPressedKeys.Contains(VirtualKeyCode);
+}
+
+void FGameInputKeyboardDeviceProcessor::SetSimulatedCapsLock(bool bVal)
+{
+	bSimulatedCapsLock = bVal;
+}
+
+void FGameInputKeyboardDeviceProcessor::ClearState(const FGameInputEventParams& Params)
+{
+	TSet<uint8> NoPressedKeys;
+	UpdateUnifiedKeyboardState(Params, NoPressedKeys);
+}
+
+GameInputKind FGameInputKeyboardDeviceProcessor::GetSupportedReadingKind() const
+{
+	return GameInputKindKeyboard;
+}
+
+static int GAllowVirtualMouseInput = 1;
+static FAutoConsoleVariableRef CVarAllowVirtualMouseInput(
+	TEXT("GameInput.AllowVirtualMouseInput"),
+	GAllowVirtualMouseInput,
+	TEXT("Whether to accept input from virtual mice, such as those from a remote viewer. Note that this doesn't change whether the mouse is 'connected'"));
+
+static float GMouseSensitivity = 1.0f;
+static FAutoConsoleVariableRef CVarMouseSensitivity(
+	TEXT("Input.MouseSensitivity"),
+	GMouseSensitivity,
+	TEXT("The sensitivity multiplier of the mouse\n")
+	TEXT(" 1 (default)"),
+	ECVF_Default);
+
+static float GMouseDoubleClickArea = 10.0f;
+static FAutoConsoleVariableRef CVarMouseDoubleClickArea(
+	TEXT("Input.DoubleClickArea"),
+	GMouseDoubleClickArea,
+	TEXT("How far the mouse can move between double clicks to still count as a double click"));
+
+static float GMouseDoubleClickDelay = 0.5f;
+static FAutoConsoleVariableRef CVarMouseDoubleClickDelay(
+	TEXT("Input.DoubleClickDelay"),
+	GMouseDoubleClickDelay,
+	TEXT("Time in seconds between mouse down events to trigger a double click event. Set to 0 to disable double clicking"));
+
+namespace UE::GameInput
+{
+	struct FMouseButtonMapping
+	{
+		uint32 GameInputButton;
+		EMouseButtons::Type MouseButton;
+	};
+	static const FMouseButtonMapping MouseButtonMappings[] =
+	{
+		{ static_cast<uint32>(GameInputMouseButtons::GameInputMouseLeftButton), EMouseButtons::Left    },
+		{ static_cast<uint32>(GameInputMouseButtons::GameInputMouseRightButton), EMouseButtons::Right   },
+		{ static_cast<uint32>(GameInputMouseButtons::GameInputMouseMiddleButton), EMouseButtons::Middle  },
+		{ static_cast<uint32>(GameInputMouseButtons::GameInputMouseButton4), EMouseButtons::Thumb01 },
+		{ static_cast<uint32>(GameInputMouseButtons::GameInputMouseButton5), EMouseButtons::Thumb02 },
+	};
+}
+
+FGameInputMouseDeviceProcessor::FGameInputMouseDeviceProcessor(const TSharedPtr<class ICursor>& InCursor)
+	: Cursor(InCursor)
+	, LastMouseOffset(ForceInitToZero)
+{
+	FMemory::Memset(PreviousMouseState, 0);
+
+	for (uint32 i = 0; i < MaxSupportedButtons; i++)
+	{
+		RepeatTime[i] = 0.0;
+	}
+}
+
+bool FGameInputMouseDeviceProcessor::CanProcessVirtualMouse() const
+{
+	// ignore the input if requested
+	if (GAllowVirtualMouseInput == 0)
+	{
+		return false;
+	}
+	return true;
+}
+
+bool FGameInputMouseDeviceProcessor::ProcessInput(const FGameInputEventParams& Params)
+{
+	// read mouse state
+	GameInputMouseState MouseState;
+	if (!Params.Reading->GetMouseState(&MouseState))
+	{
+		return false;
+	}
+
+	int32 LEGACY_VirtualMaxX = MAX_int32;
+	int32 LEGACY_VirtualMaxY = MAX_int32;
+
+	const GameInputDeviceInfo* DeviceInfo = Params.GetDeviceInfo();
+	const bool bIsVirtualMouse = DeviceInfo ? DeviceInfo->deviceFamily == GameInputFamilyVirtual : false;
+
+	const bool bHighPrecisionMouseMode = FSlateApplication::Get().IsUsingHighPrecisionMouseMovment();
+
+	if (bIsVirtualMouse)
+	{
+		if (!CanProcessVirtualMouse())
+		{
+			return false;
+		}
+	}
+
+	// update mouse position
+	float MouseDX, MouseDY = 0.0f;
+	
+	int32 RawMouseDX = static_cast<int32>(MouseState.positionX) - static_cast<int32>(PreviousMouseState.positionX);
+	int32 RawMouseDY = static_cast<int32>(MouseState.positionY) - static_cast<int32>(PreviousMouseState.positionY);	
+	
+	if (bHighPrecisionMouseMode)
+	{
+		MouseDX = static_cast<float>(RawMouseDX);
+		MouseDY = static_cast<float>(RawMouseDY);
+	}
+	else
+	{
+		MouseDX = UE::Cursor::CalculateDeltaWithAcceleration(RawMouseDX, GMouseSensitivity);
+		MouseDY = UE::Cursor::CalculateDeltaWithAcceleration(RawMouseDY, GMouseSensitivity);
+	}
+	if (MouseDX != 0 || MouseDY != 0)
+	{
+		if (Cursor.IsValid())
+		{
+			FVector2D CursorPos = Cursor->GetPosition();
+			CursorPos.X += MouseDX;
+			CursorPos.Y += MouseDY;
+			Cursor->SetPosition((int32)CursorPos.X, (int32)CursorPos.Y);
+		}
+
+		// LastMouseOffset is used for double-click detection, so we should use the same coordinates as the cursor i.e. the processed ones.
+		LastMouseOffset.X += MouseDX;
+		LastMouseOffset.Y += MouseDY;
+
+		UE_LOG(LogGameInput, VeryVerbose, TEXT("Device %s (InputDeviceId = %d) - Mouse RawDX %d, RawDY %d, DX %0.2f, DY %.2f  %s"), *UE::GameInput::LexToString(Params.Device), Params.InputDeviceId.GetId(), RawMouseDX, RawMouseDY, MouseDX, MouseDY, bHighPrecisionMouseMode ? TEXT("(high precision)") : TEXT(""));
+
+		// OnRawMouseMove may be used to process player character controls so, as the name suggests, these should be raw data without any smoothing.
+		Params.MessageHandler->OnRawMouseMove(RawMouseDX, RawMouseDY);
+	}
+
+	// update mouse wheel
+	const float MouseWheelSpinFactor = 1 / 120.0f;
+	float MouseWheelDX = (float)(MouseState.wheelX - PreviousMouseState.wheelX);
+	float MouseWheelDY = (float)(MouseState.wheelY - PreviousMouseState.wheelY);
+	if (MouseWheelDY != 0)
+	{
+		UE_LOG(LogGameInput, Verbose, TEXT("Device %s (InputDeviceId = %d) - Mouse Wheel DY %.2f"), *UE::GameInput::LexToString(Params.Device), Params.InputDeviceId.GetId(), MouseWheelDY);
+		Params.MessageHandler->OnMouseWheel(MouseWheelDY * MouseWheelSpinFactor);
+	}
+
+	// handle button change events
+	uint32 CurrentButtonHeldMask = (uint32)MouseState.buttons;
+	uint32 ActionMask = (LastButtonHeldMask ^ CurrentButtonHeldMask);
+	uint32 RepeatMask = (LastButtonHeldMask & CurrentButtonHeldMask);
+	double CurrentTime = FPlatformTime::Seconds();
+
+	for (int ButtonIndex = 0; ButtonIndex < UE_ARRAY_COUNT(UE::GameInput::MouseButtonMappings); ButtonIndex++)
+	{
+		const uint32 GameInputButton = UE::GameInput::MouseButtonMappings[ButtonIndex].GameInputButton;
+		if ((ActionMask & GameInputButton) == 0)
+		{
+			continue;
+		}
+
+		const EMouseButtons::Type MouseButton = UE::GameInput::MouseButtonMappings[ButtonIndex].MouseButton;
+		if ((CurrentButtonHeldMask & GameInputButton) != 0)
+		{
+			const bool bHasDoubleClick = (GMouseDoubleClickDelay > 0) && (CurrentTime <= RepeatTime[ButtonIndex]) && (LastMouseOffset.SizeSquared() <= FMath::Square(GMouseDoubleClickArea));
+			if (bHasDoubleClick)
+			{
+				UE_LOG(LogGameInput, Verbose, TEXT("Device %s (DeviceIndex = %d) - %s Double Click"), *UE::GameInput::LexToString(Params.Device), Params.InputDeviceId.GetId(), UE::GameInput::GetMouseButtonName(MouseButton));
+				Params.MessageHandler->OnMouseDoubleClick(nullptr, MouseButton);
+			}
+			else
+			{
+				UE_LOG(LogGameInput, Verbose, TEXT("Device %s (DeviceIndex = %d) - %s Pressed"), *UE::GameInput::LexToString(Params.Device), Params.InputDeviceId.GetId(), UE::GameInput::GetMouseButtonName(MouseButton));
+				Params.MessageHandler->OnMouseDown(nullptr, MouseButton);
+			}
+
+			RepeatTime[ButtonIndex] = CurrentTime + GMouseDoubleClickDelay;
+			LastMouseOffset = FVector2D::ZeroVector;
+		}
+		else
+		{
+			UE_LOG(LogGameInput, Verbose, TEXT("Device %s (DeviceIndex = %d) - %s Released"), *UE::GameInput::LexToString(Params.Device), Params.InputDeviceId.GetId(), UE::GameInput::GetMouseButtonName(MouseButton));
+			Params.MessageHandler->OnMouseUp(MouseButton);
+		}
+	}
+
+	PreviousMouseState = MouseState;
+	LastButtonHeldMask = CurrentButtonHeldMask;
+
+	return true;
+}
+
+void FGameInputMouseDeviceProcessor::ClearState(const FGameInputEventParams& Params)
+{
+	// clear buttons
+	for (int32 ButtonIndex = 0; ButtonIndex < UE_ARRAY_COUNT(UE::GameInput::MouseButtonMappings); ButtonIndex++)
+	{
+		uint32 GameInputButton = UE::GameInput::MouseButtonMappings[ButtonIndex].GameInputButton;
+		if ((GameInputButton & LastButtonHeldMask) != 0)
+		{
+			EMouseButtons::Type MouseButton = UE::GameInput::MouseButtonMappings[ButtonIndex].MouseButton;
+			UE_LOG(LogGameInput, Verbose, TEXT("Device %s (InputDeviceId = %d) - %s Released (via ClearState)"), *UE::GameInput::LexToString(Params.Device), Params.InputDeviceId.GetId(), UE::GameInput::GetMouseButtonName(MouseButton));
+			Params.MessageHandler->OnMouseUp(MouseButton);
+		}
+	}
+
+	// Clear repeat times
+	for (uint32 i = 0; i < MaxSupportedButtons; i++)
+	{
+		RepeatTime[i] = 0.0;
+	}
+
+	// clear previous mouse state
+	LastMouseOffset = FVector2D::ZeroVector;
+	FMemory::Memset(PreviousMouseState, 0);
+}
+
+GameInputKind FGameInputMouseDeviceProcessor::GetSupportedReadingKind() const
+{
+	return GameInputKindMouse;
+}
+
+#endif	// #if GAME_INPUT_SUPPORT
\ No newline at end of file
diff --git a/Engine/Plugins/Runtime/GameInput/Source/GameInputBase/Private/Processors/GameInputDeviceProcessor_RacingWheel.cpp b/Engine/Plugins/Runtime/GameInput/Source/GameInputBase/Private/Processors/GameInputDeviceProcessor_RacingWheel.cpp
new file mode 100644
--- /dev/null
+++ b/Engine/Plugins/Runtime/GameInput/Source/GameInputBase/Private/Processors/GameInputDeviceProcessor_RacingWheel.cpp
@@ -0,0 +1,214 @@
+// Copyright Epic Games, Inc. All Rights Reserved.
+
+#include "Processors/GameInputDeviceProcessor_RacingWheel.h"
+
+#include "GameInputDeveloperSettings.h"
+#include "GameInputKeyTypes.h"
+#include "HAL/ConsoleManager.h"
+
+#if GAME_INPUT_SUPPORT
+
+namespace UE::GameInput
+{
+	/**
+	* Returns a map of GameInputRacingWheelButtons to their associated Unreal Engine FKey names.
+	*/
+	static const TMap<uint32, FGamepadKeyNames::Type>& GetRacingWheelButtonMap()
+	{
+		static const TMap<uint32, FGamepadKeyNames::Type> RacingWheelButtonMap
+		{
+				{ static_cast<uint32>(GameInputRacingWheelButtons::GameInputRacingWheelNone), FGameInputKeys::RacingWheel_None.GetFName() },
+				{ static_cast<uint32>(GameInputRacingWheelButtons::GameInputRacingWheelMenu), FGameInputKeys::RacingWheel_Menu.GetFName() },
+				{ static_cast<uint32>(GameInputRacingWheelButtons::GameInputRacingWheelView), FGameInputKeys::RacingWheel_View.GetFName() },
+				{ static_cast<uint32>(GameInputRacingWheelButtons::GameInputRacingWheelPreviousGear), FGameInputKeys::RacingWheel_PreviousGear.GetFName() },
+				{ static_cast<uint32>(GameInputRacingWheelButtons::GameInputRacingWheelNextGear), FGameInputKeys::RacingWheel_NextGear.GetFName() },
+				{ static_cast<uint32>(GameInputRacingWheelButtons::GameInputRacingWheelDpadUp), FGamepadKeyNames::DPadUp },
+				{ static_cast<uint32>(GameInputRacingWheelButtons::GameInputRacingWheelDpadDown), FGamepadKeyNames::DPadDown },
+				{ static_cast<uint32>(GameInputRacingWheelButtons::GameInputRacingWheelDpadLeft), FGamepadKeyNames::DPadLeft },
+				{ static_cast<uint32>(GameInputRacingWheelButtons::GameInputRacingWheelDpadRight), FGamepadKeyNames::DPadRight },
+			};
+
+		return RacingWheelButtonMap;
+	}
+}; // namespace UE::GameInput
+
+FGameInputRacingWheelProcessor::FGameInputRacingWheelProcessor()
+	: IGameInputDeviceProcessor()
+{
+	NumReadingsProcessedThisFrame = 0;
+	FMemory::Memset(PreviousState, 0);
+	FMemory::Memset(RepeatTime, 0);
+}
+
+bool FGameInputRacingWheelProcessor::ProcessInput(const FGameInputEventParams& Params)
+{
+	bool bRes = false;
+
+	// Can't do anything for an invalid platform user
+	if (!Params.PlatformUserId.IsValid() || !Params.Reading)
+	{
+		return bRes;
+	}
+
+	GameInputRacingWheelState WheelState;
+	if (!Params.Reading->GetRacingWheelState(&WheelState))
+	{
+		return bRes;
+	}
+
+	// We only want to process the buttons here, as it might get called multiple times per frame.
+	bRes |= ProcessWheelButtonState(Params, WheelState);
+
+	++NumReadingsProcessedThisFrame;
+
+	return bRes;
+}
+
+bool FGameInputRacingWheelProcessor::PostProcessInput(const FGameInputEventParams& Params)
+{
+	bool bRes = false;
+
+	// Check if we have already processed buttons this frame. If we haven't we want to do it
+	const bool bHasProcessedAnyButtonsThisFrame = NumReadingsProcessedThisFrame > 0;	
+	NumReadingsProcessedThisFrame = 0;
+
+	// Can't do anything for an invalid platform user
+	if (!Params.PlatformUserId.IsValid() || !Params.PreviousReading)
+	{
+		return bRes;
+	}
+
+	// Use the "PreviousRading" because we only want to process the analog inputs once, and this will
+	// point to the most up to date reading.
+	GameInputRacingWheelState WheelState;
+	if (!Params.PreviousReading->GetRacingWheelState(&WheelState))
+	{
+		return bRes;
+	}
+
+	if (!bHasProcessedAnyButtonsThisFrame)
+	{
+		bRes |= ProcessWheelButtonState(Params, WheelState);
+	}
+
+	bRes |= ProcessWheelAnalogState(Params, WheelState);
+
+	return bRes;
+}
+
+void FGameInputRacingWheelProcessor::ClearState(const FGameInputEventParams& Params)
+{
+	// Can't do anything for an invalid platform user
+	if (!Params.PlatformUserId.IsValid())
+	{
+		return;
+	}
+
+	// We can simply process a wheel state where everything is zero
+	GameInputRacingWheelState ZeroValueState = {};
+	FMemory::Memset(ZeroValueState, 0);
+
+	ProcessWheelButtonState(Params, ZeroValueState);
+	ProcessWheelAnalogState(Params, ZeroValueState);
+
+	// Zero out any state trackers
+	NumReadingsProcessedThisFrame = 0;
+	FMemory::Memset(PreviousState, 0);
+	FMemory::Memset(RepeatTime, 0);
+}
+
+namespace UE::GameInput
+{
+	static bool HasDifferentWheelAnalogInput(const GameInputRacingWheelState& CurrentState, const GameInputRacingWheelState& PreviousState)
+	{
+		// If anything differs from the previous reading, then it has different inputs
+		return 
+			CurrentState.wheel != PreviousState.wheel ||
+			CurrentState.throttle != PreviousState.throttle ||
+			CurrentState.brake != PreviousState.brake ||
+			CurrentState.clutch != PreviousState.clutch ||
+			CurrentState.handbrake != PreviousState.handbrake;
+	}
+};
+
+const float FGameInputRacingWheelProcessor::GetRacingWheelDeadzone()
+{
+	// These settings should always be available
+	if (const UGameInputPlatformSettings* Settings = UGameInputPlatformSettings::Get())
+	{
+		return Settings->RacingWheelDeadzone;
+	}
+
+	return UGameInputPlatformSettings::DefaultRacingWheelDeadzone;
+}
+
+bool FGameInputRacingWheelProcessor::ProcessWheelAnalogState(const FGameInputEventParams& Params, GameInputRacingWheelState& CurrentWheelState)
+{
+	// ignore this input if the reading has remained empty from last time
+	const float Deadzone = GetRacingWheelDeadzone();
+
+	// If the analog values haven't changed, don't bother sending any events for them
+	const bool bHasDifferentAnalogInput = UE::GameInput::HasDifferentWheelAnalogInput(CurrentWheelState, PreviousState);
+	if (!bHasDifferentAnalogInput)
+	{
+		return false;
+	}
+
+	OnControllerAnalog(Params, FGameInputKeys::RacingWheel_Brake.GetFName(), CurrentWheelState.brake, PreviousState.brake, Deadzone);
+	OnControllerAnalog(Params, FGameInputKeys::RacingWheel_Clutch.GetFName(), CurrentWheelState.clutch, PreviousState.clutch, Deadzone);
+	OnControllerAnalog(Params, FGameInputKeys::RacingWheel_Handbrake.GetFName(), CurrentWheelState.handbrake, PreviousState.handbrake, Deadzone);
+	OnControllerAnalog(Params, FGameInputKeys::RacingWheel_Throttle.GetFName(), CurrentWheelState.throttle, PreviousState.throttle, Deadzone);
+	OnControllerAnalog(Params, FGameInputKeys::RacingWheel_Wheel.GetFName(), CurrentWheelState.wheel, PreviousState.wheel, Deadzone);
+
+	// Do we actually want this as a float? We should test the values that this can produce
+	OnControllerAnalog(Params, 
+		FGameInputKeys::RacingWheel_PatternShifterGear.GetFName(), 
+		static_cast<float>(CurrentWheelState.patternShifterGear), 
+		static_cast<float>(PreviousState.patternShifterGear), 
+		Deadzone);
+
+	// Keep track of the previous state
+	PreviousState.brake = CurrentWheelState.brake;
+	PreviousState.clutch = CurrentWheelState.clutch;
+	PreviousState.handbrake = CurrentWheelState.handbrake;
+	PreviousState.throttle = CurrentWheelState.throttle;
+	PreviousState.wheel = CurrentWheelState.wheel;
+	PreviousState.patternShifterGear = CurrentWheelState.patternShifterGear;
+
+	return true;
+}
+
+bool FGameInputRacingWheelProcessor::ProcessWheelButtonState(const FGameInputEventParams& Params, GameInputRacingWheelState& CurrentWheelState)
+{
+	// If there has been no buttons pressed on this state or the previous one, don't bother trying
+	// to evaluate any events.
+	if (CurrentWheelState.buttons == 0 && PreviousState.buttons	== 0)
+	{
+		return false;
+	}
+
+	// This might not be necessary if the racing wheels also show up as "gamepad" devices...
+	const uint32 CurrentButtonHeldMask = static_cast<uint32>(CurrentWheelState.buttons);
+
+	uint32 LastButtonHeldMask = static_cast<uint32>(PreviousState.buttons);
+
+	EvaluateButtonStates(
+		Params,
+		CurrentButtonHeldMask,
+		OUT LastButtonHeldMask,
+		RepeatTime,
+		UE::GameInput::GetRacingWheelButtonMap(),
+		MaxSupportedButtons);
+
+	// Update the previous state here
+	PreviousState.buttons = CurrentWheelState.buttons;
+
+	return true;
+}
+
+GameInputKind FGameInputRacingWheelProcessor::GetSupportedReadingKind() const
+{
+	return GameInputKindRacingWheel;
+}
+
+#endif
\ No newline at end of file
diff --git a/Engine/Plugins/Runtime/GameInput/Source/GameInputBase/Private/Processors/GameInputDeviceProcessor_Raw.cpp b/Engine/Plugins/Runtime/GameInput/Source/GameInputBase/Private/Processors/GameInputDeviceProcessor_Raw.cpp
new file mode 100644
--- /dev/null
+++ b/Engine/Plugins/Runtime/GameInput/Source/GameInputBase/Private/Processors/GameInputDeviceProcessor_Raw.cpp
@@ -0,0 +1,459 @@
+// Copyright Epic Games, Inc. All Rights Reserved.
+
+#include "Processors/GameInputDeviceProcessor_Raw.h"
+
+#include "GameInputDeveloperSettings.h"
+#include "GameInputLogging.h"
+#include "GameInputUtils.h"
+
+#if GAME_INPUT_SUPPORT
+
+#if UE_GAMEINPUT_SUPPORTS_RAW
+
+constexpr float FGameInputRawDeviceProcessor::RawValueToFloatTrigger(const uint8 RawValue) const
+{
+	// Maps the uint8 value of 0-255 to a float between 0.0 and +1.0, like a gamepad trigger.
+	return (static_cast<float>(RawValue) / 255.f);
+}
+
+const float FGameInputRawDeviceProcessor::RawValueToFloatAnalog(uint8 RawValue, const uint8 DeadZone /* = 2 */) const
+{
+	// Apply a simple square deadzone...
+	{
+		// Remember, we are mapping a uint8 (0-255) to a float of -1.0 and +1.0. So, any value
+		// between 0 and 127 is negative and 129-255 is positive. 128 is the center.
+
+		// Calculate the offset of how far this raw value is from center
+		const int32 MaxOffset = FMath::Abs(RawValue - 128);
+
+		// TODO: Implement a better deadzone then this, this is a square one
+		// If we are within the deadzone, set this value to 128 which will translate to 0.0f
+		if (MaxOffset <= DeadZone)
+		{
+			RawValue = 128;
+		}
+	}
+
+	// Maps the uint8 value to a float between -1.0 and +1.0
+	return ((static_cast<float>(RawValue) * (2.f / 255.f)) - 1.f);
+}
+
+PRAGMA_DISABLE_DEPRECATION_WARNINGS
+const GameInputRawDeviceReportInfo* FGameInputRawDeviceProcessor::ReadCurrentRawInputState(const FGameInputEventParams& Params, IGameInputReading* ReadingToUse)
+{
+#if GAMEINPUT_API_VERSION >= 3
+	UE_LOG(LogGameInput, Error, TEXT("ReadCurrentRawInputState Is deprecated in Game Input version 3+. Call ReadCurrentRawInputStateInfo instead. (Device %s) "), *UE::GameInput::LexToString(Params.Device));
+	return nullptr;
+#else
+	
+	if (!ReadingToUse)
+	{
+		UE_LOG(LogGameInput, Error, TEXT("[ReadCurrentRawInputState] Cannot read raw input state, ReadingToUse was null (Device %s) "), *UE::GameInput::LexToString(Params.Device));
+		return nullptr;
+	}
+
+	TComPtr<IGameInputRawDeviceReport> RawReport;
+	const bool bSuccessfulReading = ReadingToUse->GetRawReport(&RawReport);
+	if (!bSuccessfulReading)
+	{
+		UE_LOG(LogGameInput, Error, TEXT("[ReadCurrentRawInputState] Unsuccessful reading of raw input report! (GetRawReport failed) (Device %s) "), *UE::GameInput::LexToString(Params.Device));
+		return nullptr;
+	}
+
+	const int32 RawRepDataSize = RawReport->GetRawDataSize();
+
+	// Ensure that the current data array is populated to the size needed.
+	CurrentRawData.Reset();
+	CurrentRawData.AddZeroed(RawRepDataSize);
+
+	// Ensure that the previous data array is populated to the size needed so that we can compare its values per-index
+	if (PreviousRawData.Num() < RawRepDataSize)
+	{
+		PreviousRawData.AddZeroed(RawRepDataSize - PreviousRawData.Num());
+	}
+
+	// This will populate the values in the CurrentRawData so that we can process them
+	const int32 NumReadBytes = RawReport->GetRawData(RawReport->GetRawDataSize(), CurrentRawData.GetData());
+
+	const GameInputRawDeviceReportInfo* RawReportInfo = RawReport->GetReportInfo();
+	if (!RawReportInfo)
+	{
+		UE_LOG(LogGameInput, Warning, TEXT("[ProcessRawReport] Unsuccessful reading of raw input report! (GameInputRawDeviceReportInfo is null) (Device %s) "), *UE::GameInput::LexToString(Params.Device));
+	}
+
+	return RawReportInfo;
+#endif	// #if GAMEINPUT_API_VERSION >= 3
+}
+PRAGMA_ENABLE_DEPRECATION_WARNINGS
+
+GameInputRawDeviceReportInfo FGameInputRawDeviceProcessor::ReadCurrentRawInputStateInfo(const FGameInputEventParams& Params, IGameInputReading* ReadingToUse)
+{
+	GameInputRawDeviceReportInfo RawReportInfo = {};
+
+#if GAMEINPUT_API_VERSION >= 3
+	if (!ReadingToUse)
+	{
+		UE_LOG(LogGameInput, Error, TEXT("[ReadCurrentRawInputState] Cannot read raw input state, ReadingToUse was null (Device %s) "), *UE::GameInput::LexToString(Params.Device));
+		return RawReportInfo;
+	}
+
+	TComPtr<IGameInputRawDeviceReport> RawReport;
+	const bool bSuccessfulReading = ReadingToUse->GetRawReport(&RawReport);
+	if (!bSuccessfulReading)
+	{
+		UE_LOG(LogGameInput, Error, TEXT("[ReadCurrentRawInputState] Unsuccessful reading of raw input report! (GetRawReport failed) (Device %s) "), *UE::GameInput::LexToString(Params.Device));
+		return RawReportInfo;
+	}
+
+	const int32 RawRepDataSize = RawReport->GetRawDataSize();
+
+	// Ensure that the current data array is populated to the size needed.
+	CurrentRawData.Reset();
+	CurrentRawData.AddZeroed(RawRepDataSize);
+
+	// Ensure that the previous data array is populated to the size needed so that we can compare its values per-index
+	if (PreviousRawData.Num() < RawRepDataSize)
+	{
+		PreviousRawData.AddZeroed(RawRepDataSize - PreviousRawData.Num());
+	}
+
+	// This will populate the values in the CurrentRawData so that we can process them
+	const int32 NumReadBytes = RawReport->GetRawData(RawReport->GetRawDataSize(), CurrentRawData.GetData());
+	
+	RawReport->GetReportInfo(&RawReportInfo);
+#else
+	UE_LOG(LogGameInput, Error, TEXT("A GameInput version of 3+ is required for raw device support. (Device %s) "), *UE::GameInput::LexToString(Params.Device));
+#endif	// #if GAMEINPUT_API_VERSION >= 3
+	
+	return RawReportInfo;
+}
+
+bool FGameInputRawDeviceProcessor::ProcessInput(const FGameInputEventParams& Params)
+{
+	bool bRes = false;
+
+	// Can't do anything for an invalid platform user
+	if (!Params.PlatformUserId.IsValid())
+	{
+		return bRes;
+	}
+
+	const FGameInputDeviceConfiguration* DeviceConfig = GetDefault<UGameInputDeveloperSettings>()->FindDeviceConfiguration(Params.GetDeviceInfo());
+
+	// Check that we have a valid config before bothering to read the raw report
+	if (!DeviceConfig)
+	{
+		UE_LOG(LogGameInput, Verbose, TEXT("[ProcessRawReport] (Device %s) Does not have a valid FGameInputDeviceConfiguration in the UGameInputDeveloperSettings. We can't process Raw Input without it. Exiting."), *UE::GameInput::LexToString(Params.Device));
+		return bRes;
+	}
+
+	// Skip if this device config isn't even supposed to be processing raw input values
+	if (!DeviceConfig->bProcessRawReportData)
+	{
+		return bRes;
+	}
+
+	// Read from the current reading the CURRENT frame
+	const GameInputRawDeviceReportInfo RawReportInfo = ReadCurrentRawInputStateInfo(Params, Params.Reading);
+
+	// Actually process the current raw input values if this reading ID matches the one we want
+	if (RawReportInfo.id == DeviceConfig->RawReportReadingId)
+	{
+		// Process the BUTTON types here, we only want to process analog events once per frame to avoid over-accumulation
+		constexpr bool bShouldProcessButtons = true;
+		constexpr bool bShouldProcessAnalog = false;
+		bRes |= ProcessAllRawValues(Params, DeviceConfig, bShouldProcessButtons, bShouldProcessAnalog);
+
+		++NumReadingsProcessedThisFrame;
+	}
+
+	return bRes;
+}
+
+
+bool FGameInputRawDeviceProcessor::PostProcessInput(const FGameInputEventParams& Params)
+{
+	bool bRes = false;
+
+	const bool bProcessedAnyGamepadButtonsThisFrame = (NumReadingsProcessedThisFrame > 0);
+
+	// This is the last reading of this frame, reset the counter to 0
+	NumReadingsProcessedThisFrame = 0;
+
+	// Can't do anything for an invalid platform user
+	if (!Params.PlatformUserId.IsValid())
+	{
+		return bRes;
+	}
+
+	const FGameInputDeviceConfiguration* DeviceConfig = GetDefault<UGameInputDeveloperSettings>()->FindDeviceConfiguration(Params.GetDeviceInfo());
+
+	// Check that we have a valid config before bothering to read the raw report
+	if (!DeviceConfig)
+	{
+		UE_LOG(LogGameInput, Verbose, TEXT("[ProcessRawReport] (Device %s) Does not have a valid FGameInputDeviceConfiguration in the UGameInputDeveloperSettings. We can't process Raw Input without it. Exiting."), *UE::GameInput::LexToString(Params.Device));
+		return bRes;
+	}
+
+	// On the last input reading for the frame, the "Current Reading" should always be null. we only care for the LastReading here
+	ensure(Params.Reading == nullptr);
+
+	if (!Params.PreviousReading)
+	{
+		return bRes;
+	}
+
+	// Read from the current reading the PREVIOUS frame
+	const GameInputRawDeviceReportInfo RawReportInfo = ReadCurrentRawInputStateInfo(Params, Params.PreviousReading);
+
+	// Actually process the current raw input values if this reading ID matches the one we want
+	if (RawReportInfo.id == DeviceConfig->RawReportReadingId)
+	{
+		// We always want to process analog events on the last frame
+		constexpr bool bShouldPorcessAnalog = true;
+		// We only need to process buttons this frame too if there have been no other readings yet.
+		const bool bNeedToProcessButtons = !bProcessedAnyGamepadButtonsThisFrame;
+
+		bRes |= ProcessAllRawValues(Params, DeviceConfig, bNeedToProcessButtons, bShouldPorcessAnalog);
+	}
+
+	// Track our previous input only on the last input frame. We don't want any duplicate readings
+	PreviousRawData = CurrentRawData;
+
+	return bRes;
+}
+
+bool FGameInputRawDeviceProcessor::ProcessAllRawValues(const FGameInputEventParams& Params, const FGameInputDeviceConfiguration* DeviceConfig, const bool bShouldProcessButtons, const bool bShouldProcessAnalog)
+{
+	ensure(CurrentRawData.Num() == PreviousRawData.Num());
+
+	ensure(Params.PlatformUserId.IsValid());
+
+	bool bRes = false;
+
+	for (int32 i = 0; i < CurrentRawData.Num(); ++i)
+	{
+		const uint8 Val = CurrentRawData[i];
+
+		if (const FGameInputRawDeviceReportData* AxisData = DeviceConfig->RawReportMappingData.Find(i))
+		{
+			if (AxisData->TranslationBehavior == ERawDeviceReportTranslationBehavior::TreatAsButtonBitmask)
+			{
+				if (bShouldProcessButtons)
+				{
+					bRes |= ProcessRawInputValueAsBitmask(Params, i, AxisData);
+				}				
+			}
+			else if (AxisData->TranslationBehavior == ERawDeviceReportTranslationBehavior::TreatAsPackedAxisPair)
+			{
+				// Only call this function on the higher index
+				if (bShouldProcessAnalog && i == AxisData->HigherBitAxisIndex)
+				{
+					ProcessRawInputValueAsAanalogPaired(Params, AxisData);
+				}
+			}
+			// All other methods require a valid key name on the config
+			else if (AxisData->KeyName.IsValid())
+			{
+				// Treat this value as a button. If it is non-zero then consider it pressed. If it is zero, then it is not pressed.
+				if (AxisData->TranslationBehavior == ERawDeviceReportTranslationBehavior::TreatAsButton)
+				{
+					if (bShouldProcessButtons)
+					{
+						bRes |= ProcessRawInputValueAsButton(Params, i, AxisData);
+					}					
+				}
+				// Otherwise we can do analog values, which can be either "trigger" or "analog" types
+				else if (bShouldProcessAnalog)
+				{
+					bRes |= ProcessRawInputValueAsAanalog(Params, i, AxisData);
+				}
+			}
+			else
+			{
+				// You want a valid key name here, throw a warning if your config is wrong
+				UE_LOG(LogGameInput, Warning, TEXT("[ProcessRawReport] Invalid key name for raw report axis at index %d with value of %u (Device %s)"), i, Val, *UE::GameInput::LexToString(Params.Device));
+			}
+		}
+		else if (Val > 0)
+		{
+			UE_LOG(LogGameInput, VeryVerbose, TEXT("[ProcessRawReport] (Device %s) No raw device report config for axis '%d' with value of %u"), *UE::GameInput::LexToString(Params.Device), i, Val);
+		}
+	}
+
+	return bRes;
+}
+
+bool FGameInputRawDeviceProcessor::ProcessRawInputValueAsBitmask(const FGameInputEventParams& Params, const int32 RawValueIndex, const FGameInputRawDeviceReportData* AxisData)
+{
+	check(AxisData && AxisData->TranslationBehavior == ERawDeviceReportTranslationBehavior::TreatAsButtonBitmask);
+
+	const uint8 Val = CurrentRawData[RawValueIndex];
+	const uint8 PrevVal = PreviousRawData[RawValueIndex];
+
+	// Ensure we have a compatible key map setup
+	FPerRawInputIndexData& IndexData = RawInputIndexDataMap.FindOrAdd(RawValueIndex);
+
+	if (IndexData.KeyNameMap.IsEmpty())
+	{
+		// The settings use their button map as a bit number, so the key map we actually need to use
+		// is 1 << that bit to have EvaluateButtonStates work correctly
+		for (const TPair<int32, FName>& MappingPair : AxisData->ButtonBitMaskMappings)
+		{
+			IndexData.KeyNameMap.Add({ 1 << MappingPair.Key, MappingPair.Value });
+		}
+	}
+
+	const uint32 CurrentValue32 = static_cast<uint32>(Val);
+	uint32 PreviousValue32 = static_cast<uint32>(PrevVal);
+
+	// Note: Set the max supported buttons here to 1 because we only care about the first bit
+	EvaluateButtonStates(
+		Params,
+		CurrentValue32,
+		OUT PreviousValue32,
+		IndexData.RepeatTime,
+		IndexData.KeyNameMap,
+		FPerRawInputIndexData::MaxSupportedButtons);
+
+	// If this value is non-zero then we had a reading
+	return Val != 0;
+}
+
+bool FGameInputRawDeviceProcessor::ProcessRawInputValueAsButton(const FGameInputEventParams& Params, const int32 RawValueIndex, const FGameInputRawDeviceReportData* AxisData)
+{
+	check(AxisData && AxisData->TranslationBehavior == ERawDeviceReportTranslationBehavior::TreatAsButton);
+
+	const uint8 Val = CurrentRawData[RawValueIndex];
+	const uint8 PrevVal = PreviousRawData[RawValueIndex];
+
+	// Just treat this button as pressed when non-zero, and not pressed when 0.
+	const uint32 CurrentValue32 = FMath::Clamp<uint32>(static_cast<uint32>(Val), 0, 1);
+	uint32 PreviousValue32 = FMath::Clamp<uint32>(static_cast<uint32>(PrevVal), 0, 1);
+
+	// Map the value of the key name to this 1 if we haven't already
+	FPerRawInputIndexData& IndexData = RawInputIndexDataMap.FindOrAdd(RawValueIndex);
+	if (IndexData.KeyNameMap.IsEmpty())
+	{
+		IndexData.KeyNameMap.Add(1, AxisData->KeyName);
+	}
+
+	// Note: Set the max supported buttons here to 1 because we only care about the first bit
+	EvaluateButtonStates(
+		Params,
+		CurrentValue32,
+		OUT PreviousValue32,
+		IndexData.RepeatTime,
+		IndexData.KeyNameMap,
+		/* maxSupportedButtons */ 1);
+
+	// If this value is non-zero then we had a reading
+	return Val != 0;
+}
+
+bool FGameInputRawDeviceProcessor::ProcessRawInputValueAsAanalog(const FGameInputEventParams& Params, const int32 RawValueIndex, const FGameInputRawDeviceReportData* AxisData)
+{
+	check(AxisData);
+
+	const uint8 Val = CurrentRawData[RawValueIndex];
+	const uint8 PrevVal = PreviousRawData[RawValueIndex];
+
+	const float CurrentValueFloat = AxisData->Scalar * (AxisData->TranslationBehavior == ERawDeviceReportTranslationBehavior::TreatAsTrigger ? RawValueToFloatTrigger(Val) : RawValueToFloatAnalog(Val, AxisData->AnalogDeadzone));
+	const float PreviousValueFloat = AxisData->Scalar * (AxisData->TranslationBehavior == ERawDeviceReportTranslationBehavior::TreatAsTrigger ? RawValueToFloatTrigger(PrevVal) : RawValueToFloatAnalog(PrevVal, AxisData->AnalogDeadzone));
+
+	if (!AxisData->KeyName.IsValid())
+	{
+		return false;
+	}
+
+	OnControllerAnalog(
+		Params, 
+		AxisData->KeyName, 
+		CurrentValueFloat, 
+		PreviousValueFloat, 
+		UE::GameInput::GamepadLeftStickDeadzone, 
+		/* bShouldSetDeviceScope = */!AxisData->bIgnoreAnalogInputDeviceScopeForThisRawReport);
+
+	// We had a reading as long as it is non-zero
+	return CurrentValueFloat != 0.0f;
+}
+
+bool FGameInputRawDeviceProcessor::ProcessRawInputValueAsAanalogPaired(const FGameInputEventParams& Params, const FGameInputRawDeviceReportData* AxisData)
+{
+	check(AxisData);
+
+	if (!CurrentRawData.IsValidIndex(AxisData->LowerBitAxisIndex) || !CurrentRawData.IsValidIndex(AxisData->HigherBitAxisIndex))
+	{
+		return false;
+	}
+
+	// Get the current value
+	const uint8 CurrentLowerVal = CurrentRawData[AxisData->LowerBitAxisIndex];
+	const uint8 CurrentHigherVal = CurrentRawData[AxisData->HigherBitAxisIndex];
+	
+	// Combine the two values into a single int16. Do this by 
+	// shifting the higher value up by 8 bits, and then just use the lower value in our int16	
+	const int16 CurrentPackedVal = (int16)((CurrentHigherVal << 8) | CurrentLowerVal);
+	const float CurrentValueFloat = (static_cast<float>(CurrentPackedVal) / 32767.f);
+
+	// Get the previous value
+	const uint8 PreviousLowerVal = PreviousRawData[AxisData->LowerBitAxisIndex];
+	const uint8 PreviousHigherVal = PreviousRawData[AxisData->HigherBitAxisIndex];
+
+	const int16 PreviousPackedVal = (int16)((PreviousHigherVal << 8) | PreviousLowerVal);
+	const float PreviousValueFloat = (static_cast<float>(PreviousPackedVal) / 32767.f);
+
+	if (!AxisData->KeyName.IsValid())
+	{
+		return false;
+	}
+
+	OnControllerAnalog(Params, AxisData->KeyName, CurrentValueFloat, PreviousValueFloat, UE::GameInput::GamepadLeftStickDeadzone);
+
+	// We had a reading as long as it is non-zero
+	return CurrentValueFloat != 0.0f;
+}
+
+void FGameInputRawDeviceProcessor::ClearState(const FGameInputEventParams& Params)
+{
+	// Can't do anything for an invalid platform user
+	if (!Params.PlatformUserId.IsValid())
+	{
+		return;
+	}
+
+	const FGameInputDeviceConfiguration* DeviceConfig = GetDefault<UGameInputDeveloperSettings>()->FindDeviceConfiguration(Params.GetDeviceInfo());
+
+	// Check that we have a valid config before bothering to read the raw report
+	if (!DeviceConfig)
+	{
+		UE_LOG(LogGameInput, Verbose, TEXT("[ClearStateRawReport] No have a valid FGameInputDeviceConfiguration in the UGameInputDeveloperSettings. We can't process Raw Input without it. Exiting. (Device %s)"), *UE::GameInput::LexToString(Params.Device));
+		return;
+	}
+
+	if (!DeviceConfig->bProcessRawReportData)
+	{
+		return;
+	}
+
+	// Reset our current values to 0...
+	for (int32 i = 0; i < CurrentRawData.Num(); ++i)
+	{
+		CurrentRawData[i] = 0;
+	}
+
+	// ... and then process the raw values as if there is 0 input. We want to process all types when clearing.
+	constexpr bool bShouldPorcessButtons = true;
+	constexpr bool bShouldPorcessAnalog = true;
+
+	ProcessAllRawValues(Params, DeviceConfig, bShouldPorcessButtons, bShouldPorcessAnalog);
+}
+
+GameInputKind FGameInputRawDeviceProcessor::GetSupportedReadingKind() const
+{
+	return GameInputKindRawDeviceReport;
+}
+
+#endif	// UE_GAMEINPUT_SUPPORTS_RAW
+
+#endif	// #if GAME_INPUT_SUPPORT
diff --git a/Engine/Plugins/Runtime/GameInput/Source/GameInputBase/Private/Processors/GameInputDeviceProcessor_Sensor.cpp b/Engine/Plugins/Runtime/GameInput/Source/GameInputBase/Private/Processors/GameInputDeviceProcessor_Sensor.cpp
new file mode 100644
--- /dev/null
+++ b/Engine/Plugins/Runtime/GameInput/Source/GameInputBase/Private/Processors/GameInputDeviceProcessor_Sensor.cpp
@@ -0,0 +1,138 @@
+// Copyright Epic Games, Inc. All Rights Reserved.
+
+#include "Processors/GameInputDeviceProcessor_Sensor.h"
+
+#include "GameInputDeveloperSettings.h"
+#include "GameInputLogging.h"
+#include "GameInputUtils.h"
+#include "GameInputKeyTypes.h"
+
+#if GAME_INPUT_SUPPORT
+
+#if UE_GAMEINPUT_SUPPORTS_SENSORS
+
+bool FGameInputGamepadSensorDeviceProcessor::ProcessInput(const FGameInputEventParams& Params)
+{
+	// Do nothing in this function because all sensor data is float/analog values, so we only want to
+	// check it once per-frame and not (potentially) more. 
+	return false;
+}
+
+bool FGameInputGamepadSensorDeviceProcessor::PostProcessInput(const FGameInputEventParams& Params)
+{
+	bool bRes = false;
+
+	// Can't do anything for an invalid platform user
+	if (!Params.PlatformUserId.IsValid() || !Params.PreviousReading)
+	{
+		return bRes;
+	}
+
+	// Get the current sensor state of this frame.
+	GameInputSensorsState SensorState = {};
+	if (!Params.PreviousReading->GetSensorsState(&SensorState))
+	{
+		return bRes;
+	}
+
+	// Actually process the sensor analog data 
+	bRes |= ProcessGamepadSensorState(Params, SensorState);
+
+	// Track our previous sensor state for comparison on the next frame
+	PreviousState = SensorState;
+
+	return bRes;
+}
+
+void FGameInputGamepadSensorDeviceProcessor::ClearState(const FGameInputEventParams& Params)
+{
+	// Can't do anything for an invalid platform user
+	if (!Params.PlatformUserId.IsValid())
+	{
+		return;
+	}
+
+	// We can simply process a sensor state where everything is zero
+	GameInputSensorsState ZeroValueState = {};
+	FMemory::Memset(ZeroValueState, 0);
+
+	ProcessGamepadSensorState(Params, ZeroValueState);
+
+	// Make sure to update the previous state here
+	PreviousState = ZeroValueState;
+}
+
+bool FGameInputGamepadSensorDeviceProcessor::ProcessGamepadSensorState(const FGameInputEventParams& Params, GameInputSensorsState& SensorState)
+{
+	const GameInputDeviceInfo* Info = Params.GetDeviceInfo();
+	if (!Info || !Info->sensorsInfo)
+	{
+		return false;
+	}
+	
+	const GameInputSensorsKind DetectedSensorKind = Info->sensorsInfo->supportedSensors;
+	
+	// This device does not support any sensors. Don't bother trying to get a reading from it.
+	if (DetectedSensorKind == GameInputSensorsNone)
+	{
+		return false;
+	}
+
+	bool bHasSensorReading = false;
+	
+	static constexpr GameInputSensorsKind RequiredMotionDataKinds = (GameInputSensorsGyrometer | GameInputSensorsAccelerometer);
+	
+	// Orientation, accel, and gyro data 
+	if (DetectedSensorKind & RequiredMotionDataKinds)
+	{
+		const FVector AngularVelocity = { SensorState.angularVelocityInRadPerSecX, SensorState.angularVelocityInRadPerSecY, SensorState.angularVelocityInRadPerSecZ };
+		const FVector Accel = { SensorState.accelerationInGX, SensorState.accelerationInGY, SensorState.accelerationInGZ };
+		
+		const FQuat Orientation(SensorState.orientationW, SensorState.orientationX, SensorState.orientationY, SensorState.orientationZ);
+		const FRotator Rotation = Orientation.Rotator();
+
+		// TODO_BH: Unify motion space settings. This is using the "1" setting from other platforms 
+		
+		Params.MessageHandler->OnMotionDetected(
+				FVector(Rotation.Pitch, Rotation.Yaw, Rotation.Roll),
+				FVector(AngularVelocity.Z, -AngularVelocity.X, -AngularVelocity.Y),
+				FVector::ZeroVector,
+				FVector(-Accel.Z, Accel.X, Accel.Y),
+				Params.PlatformUserId,
+				Params.InputDeviceId
+			);
+
+		UE_LOG(LogGameInput, VeryVerbose, TEXT("[%hs] OnMotionDetected PlatUser: %d DeviceId %d  MotionSpace: %d"),
+			__func__,
+			Params.PlatformUserId.GetInternalId(),
+			Params.InputDeviceId.GetId(),
+			1);
+
+		bHasSensorReading = true;
+	}
+
+	if (DetectedSensorKind & GameInputSensorsCompass)
+	{
+		// TODO_BH: We have SensorState.headingAccuracy, so perhaps we can have a different deadzone setting for different accuracies.
+		OnControllerAnalog(
+			Params,
+			FGameInputKeys::Sensor_HeadingInDegreesFromMagneticNorth.GetFName(),
+			SensorState.headingInDegreesFromMagneticNorth,
+			PreviousState.headingInDegreesFromMagneticNorth,
+			UE::GameInput::SensorMagneticNorthDeadzone);
+
+		UE_LOG(LogGameInput, VeryVerbose, TEXT("[ProcessGamepadSensorState] (Device %s) Degrees From Magnetic North: %.2f (Accuracy: %s) "),
+			*UE::GameInput::LexToString(Params.Device), SensorState.headingInDegreesFromMagneticNorth, *UE::GameInput::LexToString(SensorState.headingAccuracy));
+	}
+	
+	return bHasSensorReading;
+}
+
+GameInputKind FGameInputGamepadSensorDeviceProcessor::GetSupportedReadingKind() const
+{
+	return GameInputKindSensors;
+}
+
+#endif	// #if UE_GAMEINPUT_SUPPORTS_SENSORS
+
+#endif	// #if GAME_INPUT_SUPPORT
\ No newline at end of file
diff --git a/Engine/Plugins/Runtime/GameInput/Source/GameInputBase/Private/Processors/GameInputDeviceProcessor_Touch.cpp b/Engine/Plugins/Runtime/GameInput/Source/GameInputBase/Private/Processors/GameInputDeviceProcessor_Touch.cpp
new file mode 100644
--- /dev/null
+++ b/Engine/Plugins/Runtime/GameInput/Source/GameInputBase/Private/Processors/GameInputDeviceProcessor_Touch.cpp
@@ -0,0 +1,153 @@
+// Copyright Epic Games, Inc. All Rights Reserved.
+
+#include "Processors/GameInputDeviceProcessor_Touch.h"
+
+#include "Framework/Application/SlateApplication.h"
+#include "GameInputLogging.h"
+#include "GameInputUtils.h"
+#include "HAL/ConsoleManager.h"
+
+#if GAME_INPUT_SUPPORT
+
+namespace GameInputDevicePreProcessorCVars
+{
+	bool bEnableTouchPressureCorrection = false;
+	FAutoConsoleVariableRef CVarEnableTouchPressureCorrection(
+		TEXT("GameInputDeviceProcessor.EnableTouchPressureCorrection"),
+		bEnableTouchPressureCorrection,
+		TEXT("This enables overriding touch input pressure in case the readings received by the touch input device preprocessor is considered too small."));
+
+	float TouchPressureCorrectionValue = 1.f;
+	FAutoConsoleVariableRef CVarTouchPressureCorrectionValue(
+		TEXT("GameInputDeviceProcessor.TouchPressureCorrectionValue"),
+		TouchPressureCorrectionValue,
+		TEXT("Touch pressure value to be used in case the touch input device preprocessor receives a reading that is considered to be too small"));
+}; // namespace GameInputDevicePreProcessorCVars
+
+#if UE_GAMEINPUT_SUPPORTS_TOUCH
+
+
+
+bool FGameInputTouchDeviceProcessor::ProcessInput(const FGameInputEventParams& Params)
+{
+	// read the new touch events
+	int32 TouchCount = Params.Reading->GetTouchCount();
+	TArray<GameInputTouchState> InputTouchData;
+	InputTouchData.SetNum(TouchCount, EAllowShrinking::No);
+	Params.Reading->GetTouchState(TouchCount, InputTouchData.GetData());
+
+	// no new touch events and we dont have any previous inputs
+	if (TouchCount == 0 && ActiveTouchPoints == 0)
+	{
+		return false;
+	}
+
+	// We can't tell slate about input messages from an invalid platform user
+	if (!Params.PlatformUserId.IsValid())
+	{
+		UE_LOG(LogGameInput, Verbose, TEXT("[FGameInputTouchDeviceProcessor::EvaluateButtonStates] Attempting to evaluate button states with an invalid platform user id of '%d'. The button messages will not be sent."), Params.PlatformUserId.GetInternalId());
+		return false;
+	}
+
+	auto FindOrAddTouch = [](TArray<FTouchData>& TouchStates, const GameInputTouchState& NewTouchData) -> int32
+	{
+		int32 Index = TouchStates.IndexOfByPredicate([&](const FTouchData& A) {return A.TouchId == NewTouchData.touchId; });
+		if (Index == INDEX_NONE)
+		{
+			// check if we have a free slot
+			Index = TouchStates.IndexOfByPredicate([&](const FTouchData& A) {return A.TouchId == INDEX_NONE; });
+
+			// no free slot - create a new one. (NB. We will never receive more than GameInputDeviceInfo::touchPointCount separate touches)
+			if (Index == INDEX_NONE)
+			{
+				Index = TouchStates.Add(FTouchData());
+			}
+		}
+		return Index;
+	};
+
+	// Reset the active touch state from the previous frame
+	for (FTouchData& TouchData : PreviousTouchData)
+	{
+		TouchData.bIsActive = false;
+	}
+
+	// process all previous frame events and check if there are in list of new events for this frame
+	for (const GameInputTouchState& InputTouch : InputTouchData)
+	{
+		// search for the same touchId in the new events list
+		int32 Index = FindOrAddTouch(PreviousTouchData, InputTouch);
+
+		if (Index == INDEX_NONE)
+		{
+			continue;
+		}
+
+		FTouchData& TouchState = PreviousTouchData[Index];
+
+		// need to scale by resolution
+		FVector2D NewPosition = FVector2D(InputTouch.positionX * (float)MaxTouchX, InputTouch.positionY * (float)MaxTouchY).RoundToVector();
+
+		// reading may fail to have a proper nonzero touch pressure, so we correct it
+		const bool bShouldCorrectTouchPressure = GameInputDevicePreProcessorCVars::bEnableTouchPressureCorrection && FMath::Abs(InputTouch.pressure) < GameInputDevicePreProcessorCVars::TouchPressureCorrectionValue;
+		const float NewTouchPressure = bShouldCorrectTouchPressure ? GameInputDevicePreProcessorCVars::TouchPressureCorrectionValue : InputTouch.pressure;
+
+		if (TouchState.TouchId == INDEX_NONE)
+		{
+			Params.MessageHandler->OnTouchStarted(nullptr, NewPosition, NewTouchPressure, Index, Params.PlatformUserId, Params.InputDeviceId);
+			TouchState.TouchId = InputTouch.touchId;
+			ActiveTouchPoints++;
+		}
+		else if (NewPosition != TouchState.Position)
+		{
+			if (!TouchState.bHasMoved)
+			{
+				Params.MessageHandler->OnTouchFirstMove(NewPosition, NewTouchPressure, Index, Params.PlatformUserId, Params.InputDeviceId);
+				TouchState.bHasMoved = true;
+			}
+			else
+			{
+				Params.MessageHandler->OnTouchMoved(NewPosition, NewTouchPressure, Index, Params.PlatformUserId, Params.InputDeviceId);
+			}
+		}
+
+		if (TouchState.Pressure != NewTouchPressure)
+		{
+			Params.MessageHandler->OnTouchForceChanged(NewPosition, NewTouchPressure, Index, Params.PlatformUserId, Params.InputDeviceId);
+		}
+
+		TouchState.Pressure = NewTouchPressure;
+		TouchState.Position = NewPosition;
+		TouchState.bIsActive = true;
+	}
+
+	// process all new event that where not in the stack
+	for (int32 Index = 0; Index < PreviousTouchData.Num(); Index++)
+	{
+		FTouchData& TouchData = PreviousTouchData[Index];
+		if (!TouchData.bIsActive && TouchData.TouchId != INDEX_NONE)
+		{
+			Params.MessageHandler->OnTouchEnded(TouchData.Position, Index, Params.PlatformUserId, Params.InputDeviceId);
+			TouchData = FTouchData();
+			ActiveTouchPoints--;
+		}
+	}
+
+	check(ActiveTouchPoints >= 0);
+
+	return true;
+}
+
+void FGameInputTouchDeviceProcessor::ClearState(const FGameInputEventParams& Params)
+{
+	// There is nothing to be done for touch data clearing...
+}
+
+GameInputKind FGameInputTouchDeviceProcessor::GetSupportedReadingKind() const
+{
+	return GameInputKindTouch;
+}
+
+#endif	// #if UE_GAMEINPUT_SUPPORTS_TOUCH
+
+#endif	// GAME_INPUT_SUPPORT
\ No newline at end of file
diff --git a/Engine/Plugins/Runtime/GameInput/Source/GameInputBase/Private/Processors/IGameInputDeviceProcessor.cpp b/Engine/Plugins/Runtime/GameInput/Source/GameInputBase/Private/Processors/IGameInputDeviceProcessor.cpp
new file mode 100644
--- /dev/null
+++ b/Engine/Plugins/Runtime/GameInput/Source/GameInputBase/Private/Processors/IGameInputDeviceProcessor.cpp
@@ -0,0 +1,270 @@
+// Copyright Epic Games, Inc. All Rights Reserved.
+
+#include "Processors/IGameInputDeviceProcessor.h"
+
+#include "GameInputDeveloperSettings.h"
+#include "GameInputLogging.h"
+#include "GameInputUtils.h"
+#include "GenericPlatform/GenericPlatformInputDeviceMapper.h"
+#include "HAL/PlatformTime.h"
+#include "Misc/ConfigCacheIni.h"
+
+#if GAME_INPUT_SUPPORT
+
+namespace UE::GameInput
+{
+	static FName InputClassName = FName("GameInput");
+}
+
+IGameInputDeviceProcessor::IGameInputDeviceProcessor()
+{
+	InitialButtonRepeatDelay = UE::GameInput::InitialRepeatDelay;
+	ButtonRepeatDelay = UE::GameInput::SubsequentRepeatDelay;
+
+	GConfig->GetFloat(TEXT("/Script/Engine.InputSettings"), TEXT("InitialButtonRepeatDelay"), InitialButtonRepeatDelay, GInputIni);
+	GConfig->GetFloat(TEXT("/Script/Engine.InputSettings"), TEXT("ButtonRepeatDelay"), ButtonRepeatDelay, GInputIni);
+}
+
+const GameInputDeviceInfo* IGameInputDeviceProcessor::FGameInputEventParams::GetDeviceInfo() const
+{
+	return UE::GameInput::GetDeviceInfo(Device);
+}
+
+bool IGameInputDeviceProcessor::PostProcessInput(const FGameInputEventParams& Params)
+{
+	// Nothing needs to be done by default here in PostProcessInput.
+	return false;
+}
+
+const FString& IGameInputDeviceProcessor::GetHardwareDeviceIdentifierName(const IGameInputDeviceProcessor::FGameInputEventParams& Params) const
+{
+	static FString ID_Virtual = TEXT("Virtual");
+	static FString ID_Aggregate = TEXT("Aggregate");
+	static FString ID_XboxOne = TEXT("XboxOne");
+	static FString ID_Xbox360 = TEXT("Xbox360");
+	static FString ID_Hid = TEXT("Hid");
+	static FString ID_I8042 = TEXT("I8042");
+
+	if (Params.Device)
+	{
+		if (const GameInputDeviceInfo* Info = Params.GetDeviceInfo())
+		{
+			// Check for any device specific overrides that may be there for custom devices
+			if (const UGameInputDeveloperSettings* Settings = GetDefault<UGameInputDeveloperSettings>())
+			{
+				if (const FGameInputDeviceConfiguration* Config = Settings->FindDeviceConfiguration(Info))
+				{
+					if (Config->bOverrideHardwareDeviceIdString)
+					{
+						return Config->OverriddenHardwareDeviceId;
+					}					
+				}
+			}
+
+			switch (Info->deviceFamily)
+			{
+			case GameInputFamilyVirtual:
+				return ID_Virtual;
+			case GameInputFamilyAggregate:
+				return ID_Aggregate;
+			case GameInputFamilyXbox360:
+				return ID_Xbox360;
+			case GameInputFamilyHid:
+				return ID_Hid;
+			case GameInputFamilyI8042:
+				return ID_I8042;
+			case GameInputFamilyXboxOne:
+			default:
+				return ID_XboxOne;
+			}
+		}
+	}
+
+	// If for some reason we are given a null device, default to just "Xbox One"
+	return ID_XboxOne;
+}
+
+void IGameInputDeviceProcessor::OnControllerAnalog(const FGameInputEventParams& Params, const FName& GamePadKey, float NewAxisValueNormalized, float OldAxisValueNormalized, float DeadZone, const bool bSetDeviceScope /*= true*/)
+{
+	if (OldAxisValueNormalized != NewAxisValueNormalized || FMath::Abs(NewAxisValueNormalized) > DeadZone)
+	{		
+		UE_LOG(LogGameInput, VeryVerbose, TEXT("Device %s (PlatformUserId = %d, InputDeviceId = %d) - Analog %s : %.3f"),
+			*UE::GameInput::LexToString(Params.Device),
+			Params.PlatformUserId.GetInternalId(),
+			Params.InputDeviceId.GetId(),
+			*GamePadKey.ToString(),
+			NewAxisValueNormalized);
+
+		// We should only tell slate about this message if the platform user and input device are valid because it will attempt
+		// to create a new slate user based on the index if it doesn't already exist
+		if (Params.PlatformUserId.IsValid() && Params.InputDeviceId.IsValid())
+		{
+			if (bSetDeviceScope)
+			{
+				FInputDeviceScope InputScope(nullptr, UE::GameInput::InputClassName, IPlatformInputDeviceMapper::Get().GetUserIndexForPlatformUser(Params.PlatformUserId), GetHardwareDeviceIdentifierName(Params));
+				Params.MessageHandler->OnControllerAnalog(GamePadKey, Params.PlatformUserId, Params.InputDeviceId, NewAxisValueNormalized);
+			}
+			else
+			{
+				Params.MessageHandler->OnControllerAnalog(GamePadKey, Params.PlatformUserId, Params.InputDeviceId, NewAxisValueNormalized);
+			}
+		}		
+	}
+}
+
+void IGameInputDeviceProcessor::EvaluateButtonStates(
+	const FGameInputEventParams& Params,
+	const uint32 CurrentButtonHeldMask,
+	uint32& PreviousButtonMask,
+	double* RepeatTime,
+	const TMap<uint32, FGamepadKeyNames::Type>& UnrealButtonNameMap,
+	const uint32 SupportedButtonCount /*= MaxSupportedButtons*/)
+{
+	// handle button change events
+	const uint32 ActionMask = (PreviousButtonMask ^ CurrentButtonHeldMask);
+	const uint32 RepeatMask = (PreviousButtonMask & CurrentButtonHeldMask);
+	uint32 BitMask = 1;
+
+	if (!RepeatTime)
+	{
+		UE_LOG(LogGameInput, Error, TEXT("[IGameInputDeviceProcessor::EvaluateButtonStates] Invalid RepeatTime array given to evaluate button states!"));
+		return;
+	}
+
+	// If the given button mask and repeat mask are both zero, then no buttons have been pressed or had a state change. No need to iterate the bitmask
+	if (ActionMask == 0 && RepeatMask == 0)
+	{
+		return;
+	}
+
+	// We can't tell slate about input messages from an invalid platform user
+	if (!Params.PlatformUserId.IsValid())
+	{
+		UE_LOG(LogGameInput, Verbose, TEXT("[IGameInputDeviceProcessor::EvaluateButtonStates] Attempting to evaluate button states with an invalid platform user id of '%d'. The button messages will not be sent."), Params.PlatformUserId.GetInternalId());
+		return;
+	}
+
+	const double CurrentTime = FPlatformTime::Seconds();
+
+	IPlatformInputDeviceMapper& DeviceMapper = IPlatformInputDeviceMapper::Get();
+
+	for (uint32 n = 0; n < SupportedButtonCount; ++n)
+	{
+		FGamepadKeyNames::Type ButtonKey;
+		if (UE::GameInput::GameInputButtonToUnrealName(UnrealButtonNameMap, BitMask, ButtonKey))
+		{
+			// Check for button state change
+			if (0 != (ActionMask & BitMask))
+			{
+				FInputDeviceScope InputScope(nullptr, UE::GameInput::InputClassName, DeviceMapper.GetUserIndexForPlatformUser(Params.PlatformUserId), GetHardwareDeviceIdentifierName(Params));
+
+				if (0 != (CurrentButtonHeldMask & BitMask))
+				{
+					UE_LOG(LogGameInput, Verbose, TEXT("[FGameInputDevice::EvaluateButtonStates] (PlatformUserId = %d, InputDeviceId = %d) - Button '%s' Pressed"), Params.PlatformUserId.GetInternalId(), Params.InputDeviceId.GetId(), *ButtonKey.ToString());
+					Params.MessageHandler->OnControllerButtonPressed(ButtonKey, Params.PlatformUserId, Params.InputDeviceId, false);
+					RepeatTime[n] = CurrentTime + InitialButtonRepeatDelay;
+				}
+				else
+				{
+					UE_LOG(LogGameInput, Verbose, TEXT("[FGameInputDevice::EvaluateButtonStates] (PlatformUserId = %d, InputDeviceId = %d) - Button '%s' Released"), Params.PlatformUserId.GetInternalId(), Params.InputDeviceId.GetId(), *ButtonKey.ToString());
+					Params.MessageHandler->OnControllerButtonReleased(ButtonKey, Params.PlatformUserId, Params.InputDeviceId, false);
+					RepeatTime[n] = 0.0;
+				}
+			}
+
+			// Check for repeat key
+			if (0 != (RepeatMask & BitMask) && CurrentTime > RepeatTime[n])
+			{
+				FInputDeviceScope InputScope(nullptr, UE::GameInput::InputClassName, DeviceMapper.GetUserIndexForPlatformUser(Params.PlatformUserId), GetHardwareDeviceIdentifierName(Params));
+
+				RepeatTime[n] = CurrentTime + ButtonRepeatDelay;
+				UE_LOG(LogGameInput, Verbose, TEXT("[FGameInputDevice::EvaluateButtonStates] (PlatformUserId = %d, InputDeviceId = %d) - Button '%s' Repeat Pressed"), Params.PlatformUserId.GetInternalId(), Params.InputDeviceId.GetId(), *ButtonKey.ToString());
+				Params.MessageHandler->OnControllerButtonPressed(ButtonKey, Params.PlatformUserId, Params.InputDeviceId, true);
+			}
+		}
+
+		// Move on to the next bit!
+		BitMask <<= 1;
+	}
+
+	PreviousButtonMask = CurrentButtonHeldMask;
+}
+
+void IGameInputDeviceProcessor::EvaluateSwitchState(
+	const FGameInputEventParams& Params,
+	GameInputSwitchPosition CurrentPosition,
+	GameInputSwitchPosition& PreviousPosition,
+	TArray<double>& RepeatTimes)
+{
+	if (!ensureMsgf(RepeatTimes.IsValidIndex(static_cast<uint32>(GameInputSwitchLeft)), TEXT("RepeatTimes array needs to be the same size as the number of switch positions!")))
+	{
+		return;
+	}
+	GameInputSwitchPosition PrevCopy = PreviousPosition;
+
+	// If the current and previous switch states are both the center, then nothing has happened.
+	if (CurrentPosition == GameInputSwitchCenter && PreviousPosition == GameInputSwitchCenter)
+	{
+		return;
+	}
+
+	// We can't send any slate input events to slate if the platform user is invalid
+	if (!Params.PlatformUserId.IsValid())
+	{
+		UE_LOG(LogGameInput, Verbose, TEXT("[IGameInputDeviceProcessor::EvaluateSwitchState] Attempting to evaluate button states with an invalid platform user id of '%d'. The input messages will not be sent."), Params.PlatformUserId.GetInternalId());
+		return;
+	}
+
+	const double CurrentTime = FPlatformTime::Seconds();
+
+	// If the current and previous are not the same, then release the previous
+	// and send the pressed event for the current, as this is the first press
+
+	if (CurrentPosition != PreviousPosition)
+	{
+		// Release the previous keys
+		if (const TArray<FGamepadKeyNames::Type>* PrevKeyArray = UE::GameInput::SwitchPositionToUnrealName(PreviousPosition))
+		{
+			for (const FGamepadKeyNames::Type KeyName : *PrevKeyArray)
+			{
+				UE_LOG(LogGameInput, Verbose, TEXT("[EvaluateSwitchState] (PlatformUserId = %d, InputDeviceId = %d) - Switch '%s' Released"), Params.PlatformUserId.GetInternalId(), Params.InputDeviceId.GetId(), *KeyName.ToString());
+				Params.MessageHandler->OnControllerButtonReleased(KeyName, Params.PlatformUserId, Params.InputDeviceId, false);
+			}
+
+			RepeatTimes[static_cast<int32>(PreviousPosition)] = 0.0;
+		}
+
+		// "Press" the current key
+		if (const TArray<FGamepadKeyNames::Type>* CurrentKeyArray = UE::GameInput::SwitchPositionToUnrealName(CurrentPosition))
+		{
+			for (const FGamepadKeyNames::Type KeyName : *CurrentKeyArray)
+			{
+				UE_LOG(LogGameInput, Verbose, TEXT("[EvaluateSwitchState] (PlatformUserId = %d, InputDeviceId = %d) - Switch '%s' Pressed"), Params.PlatformUserId.GetInternalId(), Params.InputDeviceId.GetId(), *KeyName.ToString());
+				Params.MessageHandler->OnControllerButtonPressed(KeyName, Params.PlatformUserId, Params.InputDeviceId, true);
+			}
+
+			RepeatTimes[static_cast<int32>(CurrentPosition)] = CurrentTime + InitialButtonRepeatDelay;
+		}
+	}
+	// Otherwise, the states are the same. Check if we can repeat them
+	else if (CurrentTime > RepeatTimes[static_cast<int32>(CurrentPosition)])
+	{
+		if (const TArray<FGamepadKeyNames::Type>* CurrentKeyArray = UE::GameInput::SwitchPositionToUnrealName(CurrentPosition))
+		{
+			for (const FGamepadKeyNames::Type KeyName : *CurrentKeyArray)
+			{
+				UE_LOG(LogGameInput, Verbose, TEXT("[EvaluateSwitchState] (PlatformUserId = %d, InputDeviceId = %d) - Switch '%s' Repeat Pressed"), Params.PlatformUserId.GetInternalId(), Params.InputDeviceId.GetId(), *KeyName.ToString());
+				Params.MessageHandler->OnControllerButtonPressed(KeyName, Params.PlatformUserId, Params.InputDeviceId, true);
+			}
+		}
+
+		// Check for repeat key
+		RepeatTimes[static_cast<int32>(CurrentPosition)] = CurrentTime + ButtonRepeatDelay;
+	}
+
+	// Keep track of the previous position
+	PreviousPosition = CurrentPosition;
+
+}
+
+#endif	// #if GAME_INPUT_SUPPORT
\ No newline at end of file
diff --git a/Engine/Plugins/Runtime/GameInput/Source/GameInputBase/Public/GameInputDeviceContainer.h b/Engine/Plugins/Runtime/GameInput/Source/GameInputBase/Public/GameInputDeviceContainer.h
--- a/Engine/Plugins/Runtime/GameInput/Source/GameInputBase/Public/GameInputDeviceContainer.h
+++ b/Engine/Plugins/Runtime/GameInput/Source/GameInputBase/Public/GameInputDeviceContainer.h
@@ -2,7 +2,7 @@
 
 #pragma once
 
-#include "GameInputDeviceProcessor.h"
+#include "Processors/IGameInputDeviceProcessor.h"
 #include "GenericPlatform/GenericApplicationMessageHandler.h"
 
 #define UE_API GAMEINPUTBASE_API
diff --git a/Engine/Plugins/Runtime/GameInput/Source/GameInputBase/Public/GameInputDeviceProcessor.h b/Engine/Plugins/Runtime/GameInput/Source/GameInputBase/Public/GameInputDeviceProcessor.h
--- a/Engine/Plugins/Runtime/GameInput/Source/GameInputBase/Public/GameInputDeviceProcessor.h
+++ b/Engine/Plugins/Runtime/GameInput/Source/GameInputBase/Public/GameInputDeviceProcessor.h
@@ -2,603 +2,18 @@
 
 #pragma once
 
-#include "Math/MathFwd.h"
 #include "CoreTypes.h"
-#include "Containers/Map.h"
-#include "Misc/CoreMiscDefines.h"	// For FInputDeviceId/FPlatformUserId
-#include "GenericPlatform/GenericApplicationMessageHandler.h"
-#include "GameInputBaseIncludes.h"
 
-#if GAME_INPUT_SUPPORT
-
-/**
-* A Game Input Device Processor is used to determine the state of a device
-* each frame and send messages to the given Message Handler when we want to. 
-* Each processor represents a single "kind" of Game Input device. 
-* 
-* They are created by the owning GameInput Device Container (GameInputDeviceContainer.h)
-* and polled each frame when possible. Each Human Interface Device can be made up of multiple 
-* Game Input Kinds, so a single device container may have multiple different processors associated with it.
-*/
-class IGameInputDeviceProcessor
-{
-public:
-	GAMEINPUTBASE_API IGameInputDeviceProcessor();
-
-	virtual ~IGameInputDeviceProcessor() = default;
-
-	struct FGameInputEventParams
-	{
-		/** The current reading for this input event */
-		IGameInputReading* Reading = nullptr;
-
-		/** The previous reading from the last frame. */
-		IGameInputReading* PreviousReading = nullptr;
-
-		/** Device associated with this reading */
-		IGameInputDevice* Device = nullptr;
-
-		/** The message handler to use to send any input events */
-		TSharedPtr<FGenericApplicationMessageHandler> MessageHandler;
-
-		/** The platform user that is associated with the input device that triggered this event */
-		FPlatformUserId PlatformUserId = PLATFORMUSERID_NONE;
-
-		/** The Input device ID of the device that triggered this event. */
-		FInputDeviceId InputDeviceId = INPUTDEVICEID_NONE;
-
-		/** Gets the Input Device Info for this GameInput Device if it is valid. Can be null. */
-		GAMEINPUTBASE_API const GameInputDeviceInfo* GetDeviceInfo() const;
-	};
-
-	/** 
-	* Process any input events from the given reading and send events to the message handler
-	* using the given Platform User and Input Device Id's.
-	* 
-	* This can be called multiple times per frame, like if we are running at a lower FPS.
-	* This is because GameInput may accumulate multiple readings to use by the time the game thread gets here.
-	* 
-	* If you have inputs that you only want to process once per frame, then use PostProcessInput instead.
-	* 
-	* @return True if any input events have been processed, false if nothing happened
-	*/
-	virtual bool ProcessInput(const FGameInputEventParams& Params) = 0;
-
-	/**
-	* Called after all current Game Input readings have been processed
-	* this frame and there are no more readings for a device (GetNextReading returns GAMEINPUT_E_READING_NOT_FOUND). 
-	* This is called regardless of if an accompanying ProcessInput call has happened.
-	* 
-	* This function is useful for dealing with multiple readings in a single frame for analog devices. 
-	* It is recommended that processors override this function and only process analog input events here at
-	* the end of the frame, once, instead of multiple times in ProcessInput. That is because at lower frame rates
-	* or hitches you will end up with more then one GameInput reading in the stack that we need to process. If you were to
-	* process all of those analog inputs in a single frame, the values would accumulate to above or below +-1.0,
-	* which could have unexpected behavior in Slate.
-	* 
-	* @param Params		Game Input event params sent from the owning FGameInputDeviceContainer.
-	*					Params.Reading should always be null here, as there is no current reading.
-	*					Params.PreviousReading will contain the most recent Game Input reading from the input stack.
-	* 
-	* @return			True if any input events have been processed, false if nothing happened
-	*/
-	GAMEINPUTBASE_API virtual bool PostProcessInput(const FGameInputEventParams& Params);
-
-	/** 
-	* Clear any input state related to this processor, typically by sending events to the 
-	* message handler that each FKey related to this processor now has a value of 0.
-	*/
-	virtual void ClearState(const FGameInputEventParams& Params) = 0;
-
-	/*
-	* Returns the kind of reading that this processor supports. If the current GameInput reading 
-	* is not included in this bitmask, then this processor will not not have it's "ProcessInput" 
-	* function called.
-	* 
-	* @see FGameInputDeviceContainer::ProcessInput
-	*/
-	virtual GameInputKind GetSupportedReadingKind() const = 0;
-
-protected:	
-
-	/** 
-	* Returns the string hardware device identifier for the given input event params that can be used
-	* to create an FInputDeviceScope
-	*/
-	GAMEINPUTBASE_API const FString& GetHardwareDeviceIdentifierName(const IGameInputDeviceProcessor::FGameInputEventParams& Params) const;
-
-	/** A general use function to call the message handler and tell it about a controller analog key being used */
-	GAMEINPUTBASE_API void OnControllerAnalog(const FGameInputEventParams& Params, const FName& GamePadKey, float NewAxisValueNormalized, float OldAxisValueNormalized, float DeadZone, const bool bSetDeviceScope = true);
-
-	/**
-	 * Helper function for processing the button states of Game input.
-	 *
-	 * @param CurrentButtonHeldMask			The current state of the button mask to evaluate
-	 * @param PreviousButtonMask		The previous state of the button mask. This value will be modified to store the current button mask after evaluation
-	 * @param RepeatTime				A map of the button index to a time that it was last active, used to determine if the key meets the qualifications for a repeat event
-	 * @param UnrealButtonNameMap		A map of Unreal Engine Gamepad key names to their uint32 mapping from Game Input.
-	 * @param SupportedButtonCount		The maximum number of buttons we can process
-	 */
-	GAMEINPUTBASE_API void EvaluateButtonStates(
-		const FGameInputEventParams& Params,
-		const uint32 CurrentButtonHeldMask,
-		uint32& PreviousButtonMask,
-		double* RepeatTime,
-		const TMap<uint32, FGamepadKeyNames::Type>& UnrealButtonNameMap,
-		const uint32 SupportedButtonCount);
-
-	/**
-	* Helper function to processing the state of a Switch position (aka a DPad, like left/right/up/down)
-	* 
-	* @param Params		The GameInput event params
-	* @param CurrentPosition	The current position of the switch
-	* @param PreviousPosition	The previous position of the switch. This will be set to the value of the current switch at the end of this function
-	* @param RepeatTimes		Array of doubles that represent the time at which a switch was last pressed that we can use to check for when to send repeat events
-	*/
-	GAMEINPUTBASE_API void EvaluateSwitchState(
-		const FGameInputEventParams& Params,
-		GameInputSwitchPosition CurrentPosition,
-		GameInputSwitchPosition& PreviousPosition,
-		TArray<double>& RepeatTimes);
-
-	/** Delay before sending a repeat message after a button was first pressed */
-	float InitialButtonRepeatDelay;
-
-	/** Delay before sending a repeat message after a button has been pressed for a while */
-	float ButtonRepeatDelay;
-};
-
-/**
-* Processor for Gamepad Inputs
-*/
-class FGameInputGamepadDeviceProcessor : public IGameInputDeviceProcessor
-{
-protected:
-	GAMEINPUTBASE_API virtual bool ProcessInput(const FGameInputEventParams& Params) override;
-	GAMEINPUTBASE_API virtual bool PostProcessInput(const FGameInputEventParams& Params) override;
-	GAMEINPUTBASE_API virtual void ClearState(const FGameInputEventParams& Params) override;
-	GAMEINPUTBASE_API virtual GameInputKind GetSupportedReadingKind() const override;
-
-	GAMEINPUTBASE_API bool ProcessGamepadAnalogState(const FGameInputEventParams& Params, GameInputGamepadState& GamepadState);
-	GAMEINPUTBASE_API virtual bool ProcessGamepadButtonState(const FGameInputEventParams& Params, GameInputGamepadState& GamepadState);
-
-	GameInputGamepadState PreviousState = {};
-
-	uint32 LastButtonHeldMask = 0;
-
-	uint32 LastAnalogButtonHeldMask = 0;
-
-	enum { MaxSupportedButtons = 32 };
-
-	double RepeatTime[MaxSupportedButtons];
-
-	/** 
-	* Keeps track of how many times this gamepad has been processed this frame. 
-	* Every successful processing of button input in ProcessInput will increment this value.
-	* This will get reset at the end of the input frame in PostProcessInput.
-	*/
-	int32 NumReadingsProcessedThisFrame = 0;
-};
-
-#if UE_GAMEINPUT_SUPPORTS_SENSORS
-/**
-* Processor for Gamepad "Sensor" Inputs such as gyro or accelerometers
-*/
-UE_EXPERIMENTAL(5.8, "Experimental sensor data, API may change in the future")
-class FGameInputGamepadSensorDeviceProcessor final : public IGameInputDeviceProcessor
-{
-protected:
-	virtual bool ProcessInput(const FGameInputEventParams& Params) override;
-	virtual bool PostProcessInput(const FGameInputEventParams& Params) override;
-	virtual void ClearState(const FGameInputEventParams& Params) override;
-	virtual GameInputKind GetSupportedReadingKind() const override;
-	
-	/**
-	 * Process the current sensor state and report it to the engine via the message handler.
-	 * 
-	 * @param Params Current input param data
-	 * @param CurrentState The current sensor state for this frame
-	 * @return True if there was any sensor data reported.
-	 */
-	bool ProcessGamepadSensorState(const FGameInputEventParams& Params, GameInputSensorsState& CurrentState);
-	/**
-	 * The sensor data from the previous reading that was processed
-	 */
-	GameInputSensorsState PreviousState = {};
-};
-#endif// #if UE_GAMEINPUT_SUPPORTS_SENSORS
-
-struct FGameInputDeviceConfiguration;
-
-/**
-* Processor for Controller Input Types
-* 
-* Controllers are very similar to Gamepad types, and often can just be 
-* a third party controller type that is not natively supported by the "GamepadInputKind".
-* This can include things like instruments, other platform controller types, or other 
-* third party input devices.
-*/
-class FGameInputControllerDeviceProcessor : public IGameInputDeviceProcessor
-{
-public:
-	GAMEINPUTBASE_API FGameInputControllerDeviceProcessor();
-
-protected:
-	GAMEINPUTBASE_API virtual bool ProcessInput(const FGameInputEventParams& Params) override;
-	GAMEINPUTBASE_API virtual bool PostProcessInput(const FGameInputEventParams& Params) override;
-	GAMEINPUTBASE_API virtual void ClearState(const FGameInputEventParams& Params) override;
-	GAMEINPUTBASE_API virtual GameInputKind GetSupportedReadingKind() const override;
-
-	/** Processes Controller buttons (often face buttons, like ABXY or Circle/Triangle/Square/X) */
-	GAMEINPUTBASE_API bool ProcessControllerButtonState(const FGameInputEventParams& Params, const FGameInputDeviceConfiguration* Config, IGameInputReading* InputReading);
-
-	/**
-	* Process any axis readings on this controller.
-	* This will only process axises that are configured in the given Config.
-	*/
-	GAMEINPUTBASE_API bool ProcessControllerAxisState(const FGameInputEventParams& Params, const FGameInputDeviceConfiguration* Config, IGameInputReading* InputReading);
-
-	/**
-	* Process the controller switch (also known as the DPad) state
-	*/
-	GAMEINPUTBASE_API bool ProcessControllerSwitchState(const FGameInputEventParams& Params, const FGameInputDeviceConfiguration* Config, IGameInputReading* InputReading);
-
-	/** A map of GameInputLabel types to Unreal Engine FKey names. */
-	static GAMEINPUTBASE_API const TMap<GameInputLabel, FGamepadKeyNames::Type>& GetGameInputButtonLabelToUnrealName();
-
-	/** Bitmask of any held buttons from the last frame */
-	uint32 LastButtonHeldMask = 0;
-
-	/** Max number of buttons this processor supports */
-	enum { MaxSupportedButtons = 32 };
-
-	/** Timings of when a key press should be considered a "repeat" key */
-	double RepeatTime[MaxSupportedButtons];
-
-	/** Repeat times for when switches are pressed (aka the DPad) */
-	TArray<double> SwitchRepeatTimes;
-
-	/** Previous frame's axis values */
-	TArray<float> PreviousControllerAxisValues;
-
-	/** Array of switch positions */
-	TArray<GameInputSwitchPosition> PreviousSwitchPositions;
-
-	/**
-	* Keeps track of how many times this gamepad has been processed this frame.
-	* Every successful processing of button input in ProcessInput will increment this value.
-	* This will get reset at the end of the input frame in PostProcessInput.
-	*/
-	int32 NumReadingsProcessedThisFrame = 0;
-};
-
-/**
-* Processor for Keyboard inputs
-*/
-class FGameInputKeyboardDeviceProcessor : public IGameInputDeviceProcessor
-{
-protected:
-	GAMEINPUTBASE_API virtual bool ProcessInput(const FGameInputEventParams& Params) override;
-	GAMEINPUTBASE_API virtual void ClearState(const FGameInputEventParams& Params) override;
-	GAMEINPUTBASE_API virtual GameInputKind GetSupportedReadingKind() const override;
-
-	GAMEINPUTBASE_API virtual void UpdateUnifiedKeyboardState(const FGameInputEventParams& Params, TSet<uint8>& CurrentPressedKeys);
-	
-	inline bool IsKeyPressed(uint8 VirtualKeyCode) const { return LastPressedKeys.Contains(VirtualKeyCode); }
-
-	GAMEINPUTBASE_API virtual void SetSimulatedCapsLock(bool bVal);
-
-	TSet<uint8> LastPressedKeys;	
-
-	TMap<uint8, double> KeyRepeatTime;	
-
-	// cannot read actual caps lock state, so necessary to simulate it as a bool toggle
-	bool bSimulatedCapsLock;
-};
-
-
-/**
-* Processor for mouse inputs
-*/
-class FGameInputMouseDeviceProcessor : public IGameInputDeviceProcessor
-{
-public:
-	GAMEINPUTBASE_API explicit FGameInputMouseDeviceProcessor(const TSharedPtr<class ICursor>& InCursor);
-
-protected:	
-	GAMEINPUTBASE_API virtual bool ProcessInput(const FGameInputEventParams& Params) override;
-	GAMEINPUTBASE_API virtual void ClearState(const FGameInputEventParams& Params) override;
-	GAMEINPUTBASE_API virtual GameInputKind GetSupportedReadingKind() const override;
-
-	GAMEINPUTBASE_API virtual bool CanProcessVirtualMouse() const;
-
-	TSharedPtr<class ICursor> Cursor;
-
-	GameInputMouseState PreviousMouseState;
-
-	FVector2D LastMouseOffset;
-
-	enum { MaxSupportedButtons = 32 };
-	double RepeatTime[MaxSupportedButtons];
-
-	uint32 LastButtonHeldMask;
-};
-
-
-#if UE_GAMEINPUT_SUPPORTS_TOUCH
-
-/**
-* Processor for touch inputs
-*/
-class FGameInputTouchDeviceProcessor : public IGameInputDeviceProcessor
-{
-protected:
-	GAMEINPUTBASE_API virtual bool ProcessInput(const FGameInputEventParams& Params) override;
-	GAMEINPUTBASE_API virtual void ClearState(const FGameInputEventParams& Params) override;
-	GAMEINPUTBASE_API virtual GameInputKind GetSupportedReadingKind() const override;
-
-	struct FTouchData
-	{
-		FTouchData()
-		{
-			TouchId = INDEX_NONE;
-			Position = FVector2D::ZeroVector;
-			Pressure = 0.0f;
-			bHasMoved = false;
-			bIsActive = false;
-		}
-
-		int64 TouchId;
-		FVector2D Position;
-		float Pressure;
-		bool bHasMoved;
-		bool bIsActive;
-		};
-	TArray<FTouchData> PreviousTouchData;
-	int32 ActiveTouchPoints = 0;
-
-	// The max touch X and Y of the touch region to let us scale input based on the size of the touch area
-	int32 MaxTouchX = 1920;
-	int32 MaxTouchY = 1080;
-};
-#endif	// #if UE_GAMEINPUT_SUPPORTS_TOUCH
-
-#if UE_GAMEINPUT_SUPPORTS_RAW
-
-struct FGameInputDeviceConfiguration;
-struct FGameInputRawDeviceReportData;
-
-/**
-* Processor for raw input types (GameInputKindRawDeviceReport)
-* 
-* Raw input is very customizable and is very different per-input device, so it requires a lot of configuration.
-* Most of the time it is special analog inputs on custom third party devices. A good example is the Whammy bar on
-* a guitar controller, or maybe some special shifter on a racing wheel that the Game Input API does not pick up by default.
-*/
-class FGameInputRawDeviceProcessor : public IGameInputDeviceProcessor
-{
-protected:
-	GAMEINPUTBASE_API virtual bool ProcessInput(const FGameInputEventParams& Params) override;
-	GAMEINPUTBASE_API virtual bool PostProcessInput(const FGameInputEventParams& Params) override;
-	GAMEINPUTBASE_API virtual void ClearState(const FGameInputEventParams& Params) override;
-	GAMEINPUTBASE_API virtual GameInputKind GetSupportedReadingKind() const override;
-
-	/**
-	* Reads the current state of the raw input buffer and populates CurrentRawData and PreviousRawData
-	*/
-	UE_DEPRECATED(5.9, "Use ReadCurrentRawInputStateInfo instead.")
-	GAMEINPUTBASE_API const GameInputRawDeviceReportInfo* ReadCurrentRawInputState(const FGameInputEventParams& Params, IGameInputReading* ReadingToUse);
-	
-	GAMEINPUTBASE_API GameInputRawDeviceReportInfo ReadCurrentRawInputStateInfo(const FGameInputEventParams& Params, IGameInputReading* ReadingToUse);
-
-	/**
-	* Translates the given uint8 RawValue to a float between 0.0 and +1.0, like most gamepad analog "triggers"
-	* 
-	* @param RawValue
-	* 
-	* @return			A float between 0.0 and +1.0
-	*/
-	GAMEINPUTBASE_API constexpr float RawValueToFloatTrigger(const uint8 RawValue) const;
-
-	/**
-	* Translates the given uint8 RawValue to a float between -1.0 and +1.0.
-	* 
-	* @param RawValue
-	* @param Deadzone
-	* 
-	* @return A float between -1.0 and +1.0. Will be 0.0 if the value is within the deadzone.
-	*/
-	GAMEINPUTBASE_API const float RawValueToFloatAnalog(uint8 RawValue, const uint8 DeadZone = 2) const;
-
-	/**
-	* Processes all the current raw input data. This should only be called after we have read our current raw input report
-	* or we have manually set the values (i.e. when we clear the input)
-	*/
-	GAMEINPUTBASE_API bool ProcessAllRawValues(const FGameInputEventParams& Params, const FGameInputDeviceConfiguration* DeviceConfig, const bool bShouldProcessButtons, const bool bShouldProcessAnalog);
-
-	/**
-	* Processes the raw value at the given index as a button bitmask. 
-	* For every bit in the raw value and treat it as 0 for pressed and 1 for on. 
-	*/
-	GAMEINPUTBASE_API bool ProcessRawInputValueAsBitmask(const FGameInputEventParams& Params, const int32 RawValueIndex, const FGameInputRawDeviceReportData* AxisData);
-
-	/**
-	* Processes the raw value at the given index as a button value. If the value is non-zero, then treat it as pressed. 
-	* If the value is zero, treat it as unpressed.
-	*/
-	GAMEINPUTBASE_API bool ProcessRawInputValueAsButton(const FGameInputEventParams& Params, const int32 RawValueIndex, const FGameInputRawDeviceReportData* AxisData);
-
-	/**
-	* Processes the raw value at the given index as an analog value. This will translate the uint8 to a float value 
-	* and call OnControllerAnalog for it.
-	*/
-	GAMEINPUTBASE_API bool ProcessRawInputValueAsAanalog(const FGameInputEventParams& Params, const int32 RawValueIndex, const FGameInputRawDeviceReportData* AxisData);
-
-	/**
-	* Processes a pair of raw input indexes. This allows you to combine the values of multiple uint8 axis readings and combine them
-	*/
-	GAMEINPUTBASE_API bool ProcessRawInputValueAsAanalogPaired(const FGameInputEventParams& Params, const FGameInputRawDeviceReportData* AxisData);
-
-	/**
-	* The current raw input data that is populated by IGameInputRawDeviceReport::GetRawData
-	*/
-	TArray<uint8> CurrentRawData;
-
-	/**
-	* The previous frame's raw input data that is set at the end of ProcessInput.
-	*/
-	TArray<uint8> PreviousRawData;
-
-	/** A struct that contains info about each raw input index */
-	struct FPerRawInputIndexData
-	{
-		// The max number of buttons possibly associated with this raw input data. The raw input data is
-		// a uint8, so we can have at maximum a button per-bit, so 8.
-		enum { MaxSupportedButtons = 8 };
-
-		// Stores the times that each button index has been pressed so we can detect key repeats
-		double RepeatTime[MaxSupportedButtons];
-
-		/** 
-		* A map of key names to int values that is generated dynamically based on the button mapping.
-		*/
-		TMap<uint32, FName> KeyNameMap;
-	};
-
-	/**
-	* A map of an index to some raw data associated with it. 
-	* This is only used for button evaluation, not analog.
-	*/
-	TMap<int32, FPerRawInputIndexData> RawInputIndexDataMap;
-
-	/**
-	* Keeps track of how many times this gamepad has been processed this frame.
-	* Every successful processing of button input in ProcessInput will increment this value.
-	* This will get reset at the end of the input frame in PostProcessInput.
-	*/
-	int32 NumReadingsProcessedThisFrame = 0;
-};
-
-#endif	// UE_GAMEINPUT_SUPPORTS_RAW
-
-/**
-* Processor for racing wheel types (GameInputKindRacingWheel)
-* 
-* Note: Racing wheels are often paired with other device processors like Gamepad/Controller
-* to handle any "normal" buttons that may be on them.
-*/
-class FGameInputRacingWheelProcessor : public IGameInputDeviceProcessor
-{
-public:
-	GAMEINPUTBASE_API FGameInputRacingWheelProcessor();
-protected:
-	GAMEINPUTBASE_API virtual bool ProcessInput(const FGameInputEventParams& Params) override;
-	GAMEINPUTBASE_API virtual bool PostProcessInput(const FGameInputEventParams& Params) override;
-	GAMEINPUTBASE_API virtual void ClearState(const FGameInputEventParams& Params) override;
-	GAMEINPUTBASE_API virtual GameInputKind GetSupportedReadingKind() const override;
-
-	/**
-	* Returns the deadzone that we would like to use for racing wheel analog inputs.
-	*/
-	static GAMEINPUTBASE_API const float GetRacingWheelDeadzone();
-
-	/**
-	* Processes the analog inputs of the wheel state (brake/clutch pedals, wheel movement, etc).
-	*/
-	GAMEINPUTBASE_API bool ProcessWheelAnalogState(const FGameInputEventParams& Params, GameInputRacingWheelState& CurrentWheelState);
-
-	/**
-	* Processes any buttons on the racing wheel (next/previous gears, menu nav, etc);
-	*/
-	GAMEINPUTBASE_API bool ProcessWheelButtonState(const FGameInputEventParams& Params, GameInputRacingWheelState& CurrentWheelState);
-
-	/**
-	* Keeps track of how many times this gamepad has been processed this frame.
-	* Every successful processing of button input in ProcessInput will increment this value.
-	* This will get reset at the end of the input frame in PostProcessInput.
-	*/
-	int32 NumReadingsProcessedThisFrame = 0;
-
-	/** The previous game input state that was last processed. */
-	GameInputRacingWheelState PreviousState;
-
-	/** 
-	* Array of repeat times to calculate if a button has been held long enough
-	* to receive an IE_REPEAT event.
-	*/
-	static const uint32 MaxSupportedButtons = 16;
-	double RepeatTime[MaxSupportedButtons];
-};
-
-
-/**
-* Processor for the GameInputKindArcadeStick type.
-* 
-* These are typically accessories to "fighting" style games, or other arcade style games.
-*/
-class FGameInputArcadeStickProcessor : public IGameInputDeviceProcessor
-{
-public:
-	GAMEINPUTBASE_API FGameInputArcadeStickProcessor();
-
-protected:
-
-	GAMEINPUTBASE_API virtual bool ProcessInput(const FGameInputEventParams& Params) override;
-	GAMEINPUTBASE_API virtual void ClearState(const FGameInputEventParams& Params) override;
-	GAMEINPUTBASE_API virtual GameInputKind GetSupportedReadingKind() const override;
-
-	/** The previously processed arcade stick state. */
-	GameInputArcadeStickState PreviousState;
-
-	/**
-	* Array of repeat times to calculate if a button has been held long enough
-	* to receive an IE_REPEAT event.
-	*/
-	static const uint32 MaxSupportedButtons = 16;
-	double RepeatTime[MaxSupportedButtons];
-};
-
-/**
- *  Processor for the GameInputKindFlightStick type.
- */
-class FGameInputFlightStickProcessor : public IGameInputDeviceProcessor
-{
-public:
-	GAMEINPUTBASE_API FGameInputFlightStickProcessor();
-
-protected:
-
-	GAMEINPUTBASE_API virtual bool ProcessInput(const FGameInputEventParams& Params) override;
-	GAMEINPUTBASE_API virtual bool PostProcessInput(const FGameInputEventParams& Params) override;
-	GAMEINPUTBASE_API virtual void ClearState(const FGameInputEventParams& Params) override;
-	GAMEINPUTBASE_API virtual GameInputKind GetSupportedReadingKind() const override;
-
-	/**
-	 * Process Flight stick specific buttons.
-	 */
-	GAMEINPUTBASE_API bool ProcessFlightStickButtons(const FGameInputEventParams& Params, GameInputFlightStickState& State);
-
-	/**
-	 * Process the analog values (yaw, pitch, roll, and throttle) of the flight stick
-	 */
-	GAMEINPUTBASE_API bool ProcessFlightStickAnalog(const FGameInputEventParams& Params, GameInputFlightStickState& State);
-	
-	/**
-	* Keeps track of how many times this gamepad has been processed this frame.
-	* Every successful processing of button input in ProcessInput will increment this value.
-	* This will get reset at the end of the input frame in PostProcessInput.
-	*/
-	int32 NumReadingsProcessedThisFrame = 0;
-	
-	/** State of the flight stick from the previous frame */
-	GameInputFlightStickState PreviousState = {};
-
-	// Repeat times for any pressed buttons
-	static constexpr uint32 MaxSupportedButtons = 32;
-	double RepeatTime[MaxSupportedButtons];
-
-	/** Repeat times for when switches are pressed (aka the DPad) */
-	TArray<double> SwitchRepeatTimes;
-};
-
-#endif	// GAME_INPUT_SUPPORT
+UE_DEPRECATED_HEADER(5.8, "Use the specific device processor header file you need instead of this.")
+
+// These are the new device processor header files you can use instead.
+#include "Processors/GameInputDeviceProcessor_ArcadeStick.h"
+#include "Processors/GameInputDeviceProcessor_Controller.h"
+#include "Processors/GameInputDeviceProcessor_FlightStick.h"
+#include "Processors/GameInputDeviceProcessor_Gamepad.h"
+#include "Processors/GameInputDeviceProcessor_KeyboardMouse.h"
+#include "Processors/GameInputDeviceProcessor_RacingWheel.h"
+#include "Processors/GameInputDeviceProcessor_Raw.h"
+#include "Processors/GameInputDeviceProcessor_Sensor.h"
+#include "Processors/GameInputDeviceProcessor_Touch.h"
+#include "Processors/IGameInputDeviceProcessor.h"
\ No newline at end of file
diff --git a/Engine/Plugins/Runtime/GameInput/Source/GameInputBase/Public/Processors/GameInputDeviceProcessor_ArcadeStick.h b/Engine/Plugins/Runtime/GameInput/Source/GameInputBase/Public/Processors/GameInputDeviceProcessor_ArcadeStick.h
new file mode 100644
--- /dev/null
+++ b/Engine/Plugins/Runtime/GameInput/Source/GameInputBase/Public/Processors/GameInputDeviceProcessor_ArcadeStick.h
@@ -0,0 +1,36 @@
+// Copyright Epic Games, Inc. All Rights Reserved.
+
+#pragma once
+
+#include "Processors/IGameInputDeviceProcessor.h"
+
+#if GAME_INPUT_SUPPORT
+
+/**
+* Processor for the GameInputKindArcadeStick type.
+* 
+* These are typically accessories to "fighting" style games, or other arcade style games.
+*/
+class FGameInputArcadeStickProcessor : public IGameInputDeviceProcessor
+{
+public:
+	GAMEINPUTBASE_API FGameInputArcadeStickProcessor();
+
+protected:
+
+	GAMEINPUTBASE_API virtual bool ProcessInput(const FGameInputEventParams& Params) override;
+	GAMEINPUTBASE_API virtual void ClearState(const FGameInputEventParams& Params) override;
+	GAMEINPUTBASE_API virtual GameInputKind GetSupportedReadingKind() const override;
+
+	/** The previously processed arcade stick state. */
+	GameInputArcadeStickState PreviousState;
+
+	/**
+	* Array of repeat times to calculate if a button has been held long enough
+	* to receive an IE_REPEAT event.
+	*/
+	static constexpr uint32 MaxSupportedButtons = 16;
+	double RepeatTime[MaxSupportedButtons];
+};
+
+#endif	// GAME_INPUT_SUPPORT
\ No newline at end of file
diff --git a/Engine/Plugins/Runtime/GameInput/Source/GameInputBase/Public/Processors/GameInputDeviceProcessor_Controller.h b/Engine/Plugins/Runtime/GameInput/Source/GameInputBase/Public/Processors/GameInputDeviceProcessor_Controller.h
new file mode 100644
--- /dev/null
+++ b/Engine/Plugins/Runtime/GameInput/Source/GameInputBase/Public/Processors/GameInputDeviceProcessor_Controller.h
@@ -0,0 +1,74 @@
+// Copyright Epic Games, Inc. All Rights Reserved.
+
+#pragma once
+
+#include "Processors/IGameInputDeviceProcessor.h"
+
+#if GAME_INPUT_SUPPORT
+
+struct FGameInputDeviceConfiguration;
+
+/**
+* Processor for Controller Input Types
+* 
+* Controllers are very similar to Gamepad types, and often can just be 
+* a third party controller type that is not natively supported by the "GamepadInputKind".
+* This can include things like instruments, other platform controller types, or other 
+* third party input devices.
+*/
+class FGameInputControllerDeviceProcessor : public IGameInputDeviceProcessor
+{
+public:
+	GAMEINPUTBASE_API FGameInputControllerDeviceProcessor();
+
+protected:
+	GAMEINPUTBASE_API virtual bool ProcessInput(const FGameInputEventParams& Params) override;
+	GAMEINPUTBASE_API virtual bool PostProcessInput(const FGameInputEventParams& Params) override;
+	GAMEINPUTBASE_API virtual void ClearState(const FGameInputEventParams& Params) override;
+	GAMEINPUTBASE_API virtual GameInputKind GetSupportedReadingKind() const override;
+
+	/** Processes Controller buttons (often face buttons, like ABXY or Circle/Triangle/Square/X) */
+	GAMEINPUTBASE_API bool ProcessControllerButtonState(const FGameInputEventParams& Params, const FGameInputDeviceConfiguration* Config, IGameInputReading* InputReading);
+
+	/**
+	* Process any axis readings on this controller.
+	* This will only process axises that are configured in the given Config.
+	*/
+	GAMEINPUTBASE_API bool ProcessControllerAxisState(const FGameInputEventParams& Params, const FGameInputDeviceConfiguration* Config, IGameInputReading* InputReading);
+
+	/**
+	* Process the controller switch (also known as the DPad) state
+	*/
+	GAMEINPUTBASE_API bool ProcessControllerSwitchState(const FGameInputEventParams& Params, const FGameInputDeviceConfiguration* Config, IGameInputReading* InputReading);
+
+	/** A map of GameInputLabel types to Unreal Engine FKey names. */
+	static GAMEINPUTBASE_API const TMap<GameInputLabel, FGamepadKeyNames::Type>& GetGameInputButtonLabelToUnrealName();
+
+	/** Bitmask of any held buttons from the last frame */
+	uint32 LastButtonHeldMask = 0;
+
+	/** Max number of buttons this processor supports */
+	enum { MaxSupportedButtons = 32 };
+
+	/** Timings of when a key press should be considered a "repeat" key */
+	double RepeatTime[MaxSupportedButtons];
+
+	/** Repeat times for when switches are pressed (aka the DPad) */
+	TArray<double> SwitchRepeatTimes;
+
+	/** Previous frame's axis values */
+	TArray<float> PreviousControllerAxisValues;
+
+	/** Array of switch positions */
+	TArray<GameInputSwitchPosition> PreviousSwitchPositions;
+
+	/**
+	* Keeps track of how many times this gamepad has been processed this frame.
+	* Every successful processing of button input in ProcessInput will increment this value.
+	* This will get reset at the end of the input frame in PostProcessInput.
+	*/
+	int32 NumReadingsProcessedThisFrame = 0;
+};
+
+
+#endif	// GAME_INPUT_SUPPORT
\ No newline at end of file
diff --git a/Engine/Plugins/Runtime/GameInput/Source/GameInputBase/Public/Processors/GameInputDeviceProcessor_FlightStick.h b/Engine/Plugins/Runtime/GameInput/Source/GameInputBase/Public/Processors/GameInputDeviceProcessor_FlightStick.h
new file mode 100644
--- /dev/null
+++ b/Engine/Plugins/Runtime/GameInput/Source/GameInputBase/Public/Processors/GameInputDeviceProcessor_FlightStick.h
@@ -0,0 +1,52 @@
+// Copyright Epic Games, Inc. All Rights Reserved.
+
+#pragma once
+
+#include "Processors/IGameInputDeviceProcessor.h"
+
+#if GAME_INPUT_SUPPORT
+
+/**
+ *  Processor for the GameInputKindFlightStick type.
+ */
+class FGameInputFlightStickProcessor : public IGameInputDeviceProcessor
+{
+public:
+	GAMEINPUTBASE_API FGameInputFlightStickProcessor();
+
+protected:
+
+	GAMEINPUTBASE_API virtual bool ProcessInput(const FGameInputEventParams& Params) override;
+	GAMEINPUTBASE_API virtual bool PostProcessInput(const FGameInputEventParams& Params) override;
+	GAMEINPUTBASE_API virtual void ClearState(const FGameInputEventParams& Params) override;
+	GAMEINPUTBASE_API virtual GameInputKind GetSupportedReadingKind() const override;
+
+	/**
+	 * Process Flight stick specific buttons.
+	 */
+	GAMEINPUTBASE_API bool ProcessFlightStickButtons(const FGameInputEventParams& Params, GameInputFlightStickState& State);
+
+	/**
+	 * Process the analog values (yaw, pitch, roll, and throttle) of the flight stick
+	 */
+	GAMEINPUTBASE_API bool ProcessFlightStickAnalog(const FGameInputEventParams& Params, GameInputFlightStickState& State);
+	
+	/**
+	* Keeps track of how many times this gamepad has been processed this frame.
+	* Every successful processing of button input in ProcessInput will increment this value.
+	* This will get reset at the end of the input frame in PostProcessInput.
+	*/
+	int32 NumReadingsProcessedThisFrame = 0;
+	
+	/** State of the flight stick from the previous frame */
+	GameInputFlightStickState PreviousState = {};
+
+	// Repeat times for any pressed buttons
+	static constexpr uint32 MaxSupportedButtons = 32;
+	double RepeatTime[MaxSupportedButtons];
+
+	/** Repeat times for when switches are pressed (aka the DPad) */
+	TArray<double> SwitchRepeatTimes;
+};
+
+#endif	// GAME_INPUT_SUPPORT
\ No newline at end of file
diff --git a/Engine/Plugins/Runtime/GameInput/Source/GameInputBase/Public/Processors/GameInputDeviceProcessor_Gamepad.h b/Engine/Plugins/Runtime/GameInput/Source/GameInputBase/Public/Processors/GameInputDeviceProcessor_Gamepad.h
new file mode 100644
--- /dev/null
+++ b/Engine/Plugins/Runtime/GameInput/Source/GameInputBase/Public/Processors/GameInputDeviceProcessor_Gamepad.h
@@ -0,0 +1,41 @@
+// Copyright Epic Games, Inc. All Rights Reserved.
+
+#pragma once
+
+#include "Processors/IGameInputDeviceProcessor.h"
+
+#if GAME_INPUT_SUPPORT
+
+/**
+* Processor for Gamepad Inputs
+*/
+class FGameInputGamepadDeviceProcessor : public IGameInputDeviceProcessor
+{
+protected:
+	GAMEINPUTBASE_API virtual bool ProcessInput(const FGameInputEventParams& Params) override;
+	GAMEINPUTBASE_API virtual bool PostProcessInput(const FGameInputEventParams& Params) override;
+	GAMEINPUTBASE_API virtual void ClearState(const FGameInputEventParams& Params) override;
+	GAMEINPUTBASE_API virtual GameInputKind GetSupportedReadingKind() const override;
+
+	GAMEINPUTBASE_API bool ProcessGamepadAnalogState(const FGameInputEventParams& Params, GameInputGamepadState& GamepadState);
+	GAMEINPUTBASE_API virtual bool ProcessGamepadButtonState(const FGameInputEventParams& Params, GameInputGamepadState& GamepadState);
+
+	GameInputGamepadState PreviousState = {};
+
+	uint32 LastButtonHeldMask = 0;
+
+	uint32 LastAnalogButtonHeldMask = 0;
+
+	enum { MaxSupportedButtons = 32 };
+
+	double RepeatTime[MaxSupportedButtons] = { };
+
+	/** 
+	* Keeps track of how many times this gamepad has been processed this frame. 
+	* Every successful processing of button input in ProcessInput will increment this value.
+	* This will get reset at the end of the input frame in PostProcessInput.
+	*/
+	int32 NumReadingsProcessedThisFrame = 0;
+};
+
+#endif	// GAME_INPUT_SUPPORT
\ No newline at end of file
diff --git a/Engine/Plugins/Runtime/GameInput/Source/GameInputBase/Public/Processors/GameInputDeviceProcessor_KeyboardMouse.h b/Engine/Plugins/Runtime/GameInput/Source/GameInputBase/Public/Processors/GameInputDeviceProcessor_KeyboardMouse.h
new file mode 100644
--- /dev/null
+++ b/Engine/Plugins/Runtime/GameInput/Source/GameInputBase/Public/Processors/GameInputDeviceProcessor_KeyboardMouse.h
@@ -0,0 +1,61 @@
+// Copyright Epic Games, Inc. All Rights Reserved.
+
+#pragma once
+
+#include "Processors/IGameInputDeviceProcessor.h"
+
+#if GAME_INPUT_SUPPORT
+
+/**
+* Processor for Keyboard inputs
+*/
+class FGameInputKeyboardDeviceProcessor : public IGameInputDeviceProcessor
+{
+protected:
+	GAMEINPUTBASE_API virtual bool ProcessInput(const FGameInputEventParams& Params) override;
+	GAMEINPUTBASE_API virtual void ClearState(const FGameInputEventParams& Params) override;
+	GAMEINPUTBASE_API virtual GameInputKind GetSupportedReadingKind() const override;
+
+	GAMEINPUTBASE_API virtual void UpdateUnifiedKeyboardState(const FGameInputEventParams& Params, TSet<uint8>& CurrentPressedKeys);
+	
+	GAMEINPUTBASE_API bool IsKeyPressed(uint8 VirtualKeyCode) const;
+
+	GAMEINPUTBASE_API virtual void SetSimulatedCapsLock(bool bVal);
+
+	TSet<uint8> LastPressedKeys;
+
+	TMap<uint8, double> KeyRepeatTime;	
+
+	// cannot read actual caps lock state, so necessary to simulate it as a bool toggle
+	bool bSimulatedCapsLock = {};
+};
+
+
+/**
+* Processor for mouse inputs
+*/
+class FGameInputMouseDeviceProcessor : public IGameInputDeviceProcessor
+{
+public:
+	GAMEINPUTBASE_API explicit FGameInputMouseDeviceProcessor(const TSharedPtr<class ICursor>& InCursor);
+
+protected:	
+	GAMEINPUTBASE_API virtual bool ProcessInput(const FGameInputEventParams& Params) override;
+	GAMEINPUTBASE_API virtual void ClearState(const FGameInputEventParams& Params) override;
+	GAMEINPUTBASE_API virtual GameInputKind GetSupportedReadingKind() const override;
+
+	GAMEINPUTBASE_API virtual bool CanProcessVirtualMouse() const;
+
+	TSharedPtr<class ICursor> Cursor = nullptr;
+
+	GameInputMouseState PreviousMouseState = {};
+
+	FVector2D LastMouseOffset = {};
+
+	enum { MaxSupportedButtons = 32 };
+	double RepeatTime[MaxSupportedButtons];
+
+	uint32 LastButtonHeldMask = 0;
+};
+
+#endif	// GAME_INPUT_SUPPORT
\ No newline at end of file
diff --git a/Engine/Plugins/Runtime/GameInput/Source/GameInputBase/Public/Processors/GameInputDeviceProcessor_RacingWheel.h b/Engine/Plugins/Runtime/GameInput/Source/GameInputBase/Public/Processors/GameInputDeviceProcessor_RacingWheel.h
new file mode 100644
--- /dev/null
+++ b/Engine/Plugins/Runtime/GameInput/Source/GameInputBase/Public/Processors/GameInputDeviceProcessor_RacingWheel.h
@@ -0,0 +1,59 @@
+// Copyright Epic Games, Inc. All Rights Reserved.
+
+#pragma once
+
+#include "Processors/IGameInputDeviceProcessor.h"
+
+#if GAME_INPUT_SUPPORT
+
+/**
+* Processor for racing wheel types (GameInputKindRacingWheel)
+* 
+* Note: Racing wheels are often paired with other device processors like Gamepad/Controller
+* to handle any "normal" buttons that may be on them.
+*/
+class FGameInputRacingWheelProcessor : public IGameInputDeviceProcessor
+{
+public:
+	GAMEINPUTBASE_API FGameInputRacingWheelProcessor();
+protected:
+	GAMEINPUTBASE_API virtual bool ProcessInput(const FGameInputEventParams& Params) override;
+	GAMEINPUTBASE_API virtual bool PostProcessInput(const FGameInputEventParams& Params) override;
+	GAMEINPUTBASE_API virtual void ClearState(const FGameInputEventParams& Params) override;
+	GAMEINPUTBASE_API virtual GameInputKind GetSupportedReadingKind() const override;
+
+	/**
+	* Returns the deadzone that we would like to use for racing wheel analog inputs.
+	*/
+	static GAMEINPUTBASE_API const float GetRacingWheelDeadzone();
+
+	/**
+	* Processes the analog inputs of the wheel state (brake/clutch pedals, wheel movement, etc).
+	*/
+	GAMEINPUTBASE_API bool ProcessWheelAnalogState(const FGameInputEventParams& Params, GameInputRacingWheelState& CurrentWheelState);
+
+	/**
+	* Processes any buttons on the racing wheel (next/previous gears, menu nav, etc);
+	*/
+	GAMEINPUTBASE_API bool ProcessWheelButtonState(const FGameInputEventParams& Params, GameInputRacingWheelState& CurrentWheelState);
+
+	/**
+	* Keeps track of how many times this gamepad has been processed this frame.
+	* Every successful processing of button input in ProcessInput will increment this value.
+	* This will get reset at the end of the input frame in PostProcessInput.
+	*/
+	int32 NumReadingsProcessedThisFrame = 0;
+
+	/** The previous game input state that was last processed. */
+	GameInputRacingWheelState PreviousState;
+
+	/** 
+	* Array of repeat times to calculate if a button has been held long enough
+	* to receive an IE_REPEAT event.
+	*/
+	static const uint32 MaxSupportedButtons = 16;
+	double RepeatTime[MaxSupportedButtons];
+};
+
+
+#endif	// GAME_INPUT_SUPPORT
\ No newline at end of file
diff --git a/Engine/Plugins/Runtime/GameInput/Source/GameInputBase/Public/Processors/GameInputDeviceProcessor_Raw.h b/Engine/Plugins/Runtime/GameInput/Source/GameInputBase/Public/Processors/GameInputDeviceProcessor_Raw.h
new file mode 100644
--- /dev/null
+++ b/Engine/Plugins/Runtime/GameInput/Source/GameInputBase/Public/Processors/GameInputDeviceProcessor_Raw.h
@@ -0,0 +1,127 @@
+// Copyright Epic Games, Inc. All Rights Reserved.
+
+#pragma once
+
+#include "Processors/IGameInputDeviceProcessor.h"
+
+#if GAME_INPUT_SUPPORT
+
+struct FGameInputDeviceConfiguration;
+struct FGameInputRawDeviceReportData;
+
+#if UE_GAMEINPUT_SUPPORTS_RAW
+
+/**
+* Processor for raw input types (GameInputKindRawDeviceReport)
+* 
+* Raw input is very customizable and is very different per-input device, so it requires a lot of configuration.
+* Most of the time it is special analog inputs on custom third party devices. A good example is the Whammy bar on
+* a guitar controller, or maybe some special shifter on a racing wheel that the Game Input API does not pick up by default.
+*/
+class FGameInputRawDeviceProcessor : public IGameInputDeviceProcessor
+{
+protected:
+	GAMEINPUTBASE_API virtual bool ProcessInput(const FGameInputEventParams& Params) override;
+	GAMEINPUTBASE_API virtual bool PostProcessInput(const FGameInputEventParams& Params) override;
+	GAMEINPUTBASE_API virtual void ClearState(const FGameInputEventParams& Params) override;
+	GAMEINPUTBASE_API virtual GameInputKind GetSupportedReadingKind() const override;
+
+	/**
+	* Reads the current state of the raw input buffer and populates CurrentRawData and PreviousRawData
+	*/
+	UE_DEPRECATED(5.9, "Use ReadCurrentRawInputStateInfo instead.")
+	GAMEINPUTBASE_API const GameInputRawDeviceReportInfo* ReadCurrentRawInputState(const FGameInputEventParams& Params, IGameInputReading* ReadingToUse);
+	
+	GAMEINPUTBASE_API GameInputRawDeviceReportInfo ReadCurrentRawInputStateInfo(const FGameInputEventParams& Params, IGameInputReading* ReadingToUse);
+
+	/**
+	* Translates the given uint8 RawValue to a float between 0.0 and +1.0, like most gamepad analog "triggers"
+	* 
+	* @param RawValue
+	* 
+	* @return			A float between 0.0 and +1.0
+	*/
+	GAMEINPUTBASE_API constexpr float RawValueToFloatTrigger(const uint8 RawValue) const;
+
+	/**
+	* Translates the given uint8 RawValue to a float between -1.0 and +1.0.
+	* 
+	* @param RawValue
+	* @param Deadzone
+	* 
+	* @return A float between -1.0 and +1.0. Will be 0.0 if the value is within the deadzone.
+	*/
+	GAMEINPUTBASE_API const float RawValueToFloatAnalog(uint8 RawValue, const uint8 DeadZone = 2) const;
+
+	/**
+	* Processes all the current raw input data. This should only be called after we have read our current raw input report
+	* or we have manually set the values (i.e. when we clear the input)
+	*/
+	GAMEINPUTBASE_API bool ProcessAllRawValues(const FGameInputEventParams& Params, const FGameInputDeviceConfiguration* DeviceConfig, const bool bShouldProcessButtons, const bool bShouldProcessAnalog);
+
+	/**
+	* Processes the raw value at the given index as a button bitmask. 
+	* For every bit in the raw value and treat it as 0 for pressed and 1 for on. 
+	*/
+	GAMEINPUTBASE_API bool ProcessRawInputValueAsBitmask(const FGameInputEventParams& Params, const int32 RawValueIndex, const FGameInputRawDeviceReportData* AxisData);
+
+	/**
+	* Processes the raw value at the given index as a button value. If the value is non-zero, then treat it as pressed. 
+	* If the value is zero, treat it as unpressed.
+	*/
+	GAMEINPUTBASE_API bool ProcessRawInputValueAsButton(const FGameInputEventParams& Params, const int32 RawValueIndex, const FGameInputRawDeviceReportData* AxisData);
+
+	/**
+	* Processes the raw value at the given index as an analog value. This will translate the uint8 to a float value 
+	* and call OnControllerAnalog for it.
+	*/
+	GAMEINPUTBASE_API bool ProcessRawInputValueAsAanalog(const FGameInputEventParams& Params, const int32 RawValueIndex, const FGameInputRawDeviceReportData* AxisData);
+
+	/**
+	* Processes a pair of raw input indexes. This allows you to combine the values of multiple uint8 axis readings and combine them
+	*/
+	GAMEINPUTBASE_API bool ProcessRawInputValueAsAanalogPaired(const FGameInputEventParams& Params, const FGameInputRawDeviceReportData* AxisData);
+
+	/**
+	* The current raw input data that is populated by IGameInputRawDeviceReport::GetRawData
+	*/
+	TArray<uint8> CurrentRawData;
+
+	/**
+	* The previous frame's raw input data that is set at the end of ProcessInput.
+	*/
+	TArray<uint8> PreviousRawData;
+
+	/** A struct that contains info about each raw input index */
+	struct FPerRawInputIndexData
+	{
+		// The max number of buttons possibly associated with this raw input data. The raw input data is
+		// a uint8, so we can have at maximum a button per-bit, so 8.
+		enum { MaxSupportedButtons = 8 };
+
+		// Stores the times that each button index has been pressed so we can detect key repeats
+		double RepeatTime[MaxSupportedButtons];
+
+		/** 
+		* A map of key names to int values that is generated dynamically based on the button mapping.
+		*/
+		TMap<uint32, FName> KeyNameMap;
+	};
+
+	/**
+	* A map of an index to some raw data associated with it. 
+	* This is only used for button evaluation, not analog.
+	*/
+	TMap<int32, FPerRawInputIndexData> RawInputIndexDataMap;
+
+	/**
+	* Keeps track of how many times this gamepad has been processed this frame.
+	* Every successful processing of button input in ProcessInput will increment this value.
+	* This will get reset at the end of the input frame in PostProcessInput.
+	*/
+	int32 NumReadingsProcessedThisFrame = 0;
+};
+
+#endif	// UE_GAMEINPUT_SUPPORTS_RAW
+
+#endif	// GAME_INPUT_SUPPORT
\ No newline at end of file
diff --git a/Engine/Plugins/Runtime/GameInput/Source/GameInputBase/Public/Processors/GameInputDeviceProcessor_Sensor.h b/Engine/Plugins/Runtime/GameInput/Source/GameInputBase/Public/Processors/GameInputDeviceProcessor_Sensor.h
new file mode 100644
--- /dev/null
+++ b/Engine/Plugins/Runtime/GameInput/Source/GameInputBase/Public/Processors/GameInputDeviceProcessor_Sensor.h
@@ -0,0 +1,37 @@
+// Copyright Epic Games, Inc. All Rights Reserved.
+
+#pragma once
+
+#include "Processors/IGameInputDeviceProcessor.h"
+
+#if GAME_INPUT_SUPPORT
+
+#if UE_GAMEINPUT_SUPPORTS_SENSORS
+/**
+* Processor for Gamepad "Sensor" Inputs such as gyro or accelerometers
+*/
+UE_EXPERIMENTAL(5.8, "Experimental sensor data, API may change in the future")
+class FGameInputGamepadSensorDeviceProcessor final : public IGameInputDeviceProcessor
+{
+protected:
+	virtual bool ProcessInput(const FGameInputEventParams& Params) override;
+	virtual bool PostProcessInput(const FGameInputEventParams& Params) override;
+	virtual void ClearState(const FGameInputEventParams& Params) override;
+	virtual GameInputKind GetSupportedReadingKind() const override;
+	
+	/**
+	 * Process the current sensor state and report it to the engine via the message handler.
+	 * 
+	 * @param Params Current input param data
+	 * @param CurrentState The current sensor state for this frame
+	 * @return True if there was any sensor data reported.
+	 */
+	bool ProcessGamepadSensorState(const FGameInputEventParams& Params, GameInputSensorsState& CurrentState);
+	/**
+	 * The sensor data from the previous reading that was processed
+	 */
+	GameInputSensorsState PreviousState = {};
+};
+#endif// #if UE_GAMEINPUT_SUPPORTS_SENSORS
+
+#endif	// GAME_INPUT_SUPPORT
\ No newline at end of file
diff --git a/Engine/Plugins/Runtime/GameInput/Source/GameInputBase/Public/Processors/GameInputDeviceProcessor_Touch.h b/Engine/Plugins/Runtime/GameInput/Source/GameInputBase/Public/Processors/GameInputDeviceProcessor_Touch.h
new file mode 100644
--- /dev/null
+++ b/Engine/Plugins/Runtime/GameInput/Source/GameInputBase/Public/Processors/GameInputDeviceProcessor_Touch.h
@@ -0,0 +1,48 @@
+// Copyright Epic Games, Inc. All Rights Reserved.
+
+#pragma once
+
+#include "Processors/IGameInputDeviceProcessor.h"
+
+#if GAME_INPUT_SUPPORT
+
+#if UE_GAMEINPUT_SUPPORTS_TOUCH
+
+/**
+* Processor for touch inputs
+*/
+class FGameInputTouchDeviceProcessor : public IGameInputDeviceProcessor
+{
+protected:
+	GAMEINPUTBASE_API virtual bool ProcessInput(const FGameInputEventParams& Params) override;
+	GAMEINPUTBASE_API virtual void ClearState(const FGameInputEventParams& Params) override;
+	GAMEINPUTBASE_API virtual GameInputKind GetSupportedReadingKind() const override;
+
+	struct FTouchData
+	{
+		FTouchData()
+		{
+			TouchId = INDEX_NONE;
+			Position = FVector2D::ZeroVector;
+			Pressure = 0.0f;
+			bHasMoved = false;
+			bIsActive = false;
+		}
+
+		int64 TouchId;
+		FVector2D Position;
+		float Pressure;
+		bool bHasMoved;
+		bool bIsActive;
+	};
+	TArray<FTouchData> PreviousTouchData;
+	int32 ActiveTouchPoints = 0;
+
+	// The max touch X and Y of the touch region to let us scale input based on the size of the touch area
+	int32 MaxTouchX = 1920;
+	int32 MaxTouchY = 1080;
+};
+#endif	// #if UE_GAMEINPUT_SUPPORTS_TOUCH
+
+
+#endif	// GAME_INPUT_SUPPORT
\ No newline at end of file
diff --git a/Engine/Plugins/Runtime/GameInput/Source/GameInputBase/Public/Processors/IGameInputDeviceProcessor.h b/Engine/Plugins/Runtime/GameInput/Source/GameInputBase/Public/Processors/IGameInputDeviceProcessor.h
new file mode 100644
--- /dev/null
+++ b/Engine/Plugins/Runtime/GameInput/Source/GameInputBase/Public/Processors/IGameInputDeviceProcessor.h
@@ -0,0 +1,149 @@
+// Copyright Epic Games, Inc. All Rights Reserved.
+
+#pragma once
+
+#include "Containers/Map.h"
+#include "Misc/CoreMiscDefines.h"
+#include "GenericPlatform/GenericApplicationMessageHandler.h"
+#include "GameInputBaseIncludes.h"
+
+#if GAME_INPUT_SUPPORT
+
+/**
+* A Game Input Device Processor is used to determine the state of a device
+* each frame and send messages to the given Message Handler when we want to. 
+* Each processor represents a single "kind" of Game Input device. 
+* 
+* They are created by the owning GameInput Device Container (GameInputDeviceContainer.h)
+* and polled each frame when possible. Each Human Interface Device can be made up of multiple 
+* Game Input Kinds, so a single device container may have multiple different processors associated with it.
+*/
+class IGameInputDeviceProcessor
+{
+public:
+	GAMEINPUTBASE_API IGameInputDeviceProcessor();
+
+	virtual ~IGameInputDeviceProcessor() = default;
+
+	struct FGameInputEventParams
+	{
+		/** The current reading for this input event */
+		IGameInputReading* Reading = nullptr;
+
+		/** The previous reading from the last frame. */
+		IGameInputReading* PreviousReading = nullptr;
+
+		/** Device associated with this reading */
+		IGameInputDevice* Device = nullptr;
+
+		/** The message handler to use to send any input events */
+		TSharedPtr<FGenericApplicationMessageHandler> MessageHandler = nullptr;
+
+		/** The platform user that is associated with the input device that triggered this event */
+		FPlatformUserId PlatformUserId = PLATFORMUSERID_NONE;
+
+		/** The Input device ID of the device that triggered this event. */
+		FInputDeviceId InputDeviceId = INPUTDEVICEID_NONE;
+
+		/** Gets the Input Device Info for this GameInput Device if it is valid. Can be null. */
+		GAMEINPUTBASE_API const GameInputDeviceInfo* GetDeviceInfo() const;
+	};
+
+	/** 
+	* Process any input events from the given reading and send events to the message handler
+	* using the given Platform User and Input Device Id's.
+	* 
+	* This can be called multiple times per frame, like if we are running at a lower FPS.
+	* This is because GameInput may accumulate multiple readings to use by the time the game thread gets here.
+	* 
+	* If you have inputs that you only want to process once per frame, then use PostProcessInput instead.
+	* 
+	* @return True if any input events have been processed, false if nothing happened
+	*/
+	virtual bool ProcessInput(const FGameInputEventParams& Params) = 0;
+
+	/**
+	* Called after all current Game Input readings have been processed
+	* this frame and there are no more readings for a device (GetNextReading returns GAMEINPUT_E_READING_NOT_FOUND). 
+	* This is called regardless of if an accompanying ProcessInput call has happened.
+	* 
+	* This function is useful for dealing with multiple readings in a single frame for analog devices. 
+	* It is recommended that processors override this function and only process analog input events here at
+	* the end of the frame, once, instead of multiple times in ProcessInput. That is because at lower frame rates
+	* or hitches you will end up with more then one GameInput reading in the stack that we need to process. If you were to
+	* process all of those analog inputs in a single frame, the values would accumulate to above or below +-1.0,
+	* which could have unexpected behavior in Slate.
+	* 
+	* @param Params		Game Input event params sent from the owning FGameInputDeviceContainer.
+	*					Params.Reading should always be null here, as there is no current reading.
+	*					Params.PreviousReading will contain the most recent Game Input reading from the input stack.
+	* 
+	* @return			True if any input events have been processed, false if nothing happened
+	*/
+	GAMEINPUTBASE_API virtual bool PostProcessInput(const FGameInputEventParams& Params);
+
+	/** 
+	* Clear any input state related to this processor, typically by sending events to the 
+	* message handler that each FKey related to this processor now has a value of 0.
+	*/
+	virtual void ClearState(const FGameInputEventParams& Params) = 0;
+
+	/*
+	* Returns the kind of reading that this processor supports. If the current GameInput reading 
+	* is not included in this bitmask, then this processor will not have it's "ProcessInput" 
+	* function called.
+	* 
+	* @see FGameInputDeviceContainer::ProcessInput
+	*/
+	virtual GameInputKind GetSupportedReadingKind() const = 0;
+
+protected:	
+
+	/** 
+	* Returns the string hardware device identifier for the given input event params that can be used
+	* to create an FInputDeviceScope
+	*/
+	GAMEINPUTBASE_API const FString& GetHardwareDeviceIdentifierName(const IGameInputDeviceProcessor::FGameInputEventParams& Params) const;
+
+	/** A general use function to call the message handler and tell it about a controller analog key being used */
+	GAMEINPUTBASE_API void OnControllerAnalog(const FGameInputEventParams& Params, const FName& GamePadKey, float NewAxisValueNormalized, float OldAxisValueNormalized, float DeadZone, const bool bSetDeviceScope = true);
+
+	/**
+	 * Helper function for processing the button states of Game input.
+	 *
+	 * @param CurrentButtonHeldMask			The current state of the button mask to evaluate
+	 * @param PreviousButtonMask		The previous state of the button mask. This value will be modified to store the current button mask after evaluation
+	 * @param RepeatTime				A map of the button index to a time that it was last active, used to determine if the key meets the qualifications for a repeat event
+	 * @param UnrealButtonNameMap		A map of Unreal Engine Gamepad key names to their uint32 mapping from Game Input.
+	 * @param SupportedButtonCount		The maximum number of buttons we can process
+	 */
+	GAMEINPUTBASE_API void EvaluateButtonStates(
+		const FGameInputEventParams& Params,
+		const uint32 CurrentButtonHeldMask,
+		uint32& PreviousButtonMask,
+		double* RepeatTime,
+		const TMap<uint32, FGamepadKeyNames::Type>& UnrealButtonNameMap,
+		const uint32 SupportedButtonCount);
+
+	/**
+	* Helper function to processing the state of a Switch position (aka a DPad, like left/right/up/down)
+	* 
+	* @param Params		The GameInput event params
+	* @param CurrentPosition	The current position of the switch
+	* @param PreviousPosition	The previous position of the switch. This will be set to the value of the current switch at the end of this function
+	* @param RepeatTimes		Array of doubles that represent the time at which a switch was last pressed that we can use to check for when to send repeat events
+	*/
+	GAMEINPUTBASE_API void EvaluateSwitchState(
+		const FGameInputEventParams& Params,
+		GameInputSwitchPosition CurrentPosition,
+		GameInputSwitchPosition& PreviousPosition,
+		TArray<double>& RepeatTimes);
+
+	/** Delay before sending a repeat message after a button was first pressed */
+	float InitialButtonRepeatDelay;
+
+	/** Delay before sending a repeat message after a button has been pressed for a while */
+	float ButtonRepeatDelay;
+};
+
+#endif	// GAME_INPUT_SUPPORT
