[Game Input] Fix for the switch(aka DPad) state incorrectly reporting a "repeat" trigger event for "Controller" and "Flight stick" types.

#rnx
#Synced-CL 48353899

diff --git a/Engine/Plugins/Runtime/GameInput/Source/GameInputBase/Private/Processors/IGameInputDeviceProcessor.cpp b/Engine/Plugins/Runtime/GameInput/Source/GameInputBase/Private/Processors/IGameInputDeviceProcessor.cpp
--- a/Engine/Plugins/Runtime/GameInput/Source/GameInputBase/Private/Processors/IGameInputDeviceProcessor.cpp
+++ b/Engine/Plugins/Runtime/GameInput/Source/GameInputBase/Private/Processors/IGameInputDeviceProcessor.cpp
@@ -240,7 +240,7 @@ void IGameInputDeviceProcessor::EvaluateSwitchState(
 			for (const FGamepadKeyNames::Type KeyName : *CurrentKeyArray)
 			{
 				UE_LOG(LogGameInput, Verbose, TEXT("[EvaluateSwitchState] (PlatformUserId = %d, InputDeviceId = %d) - Switch '%s' Pressed"), Params.PlatformUserId.GetInternalId(), Params.InputDeviceId.GetId(), *KeyName.ToString());
-				Params.MessageHandler->OnControllerButtonPressed(KeyName, Params.PlatformUserId, Params.InputDeviceId, true);
+				Params.MessageHandler->OnControllerButtonPressed(KeyName, Params.PlatformUserId, Params.InputDeviceId, false);
 			}
 
 			RepeatTimes[static_cast<int32>(CurrentPosition)] = CurrentTime + InitialButtonRepeatDelay;
