[Game Input] Fix typo in log

#rnx
#Synced-CL 49409908

diff --git a/Engine/Plugins/Runtime/GameInput/Source/GameInputBase/Private/GameInputDeviceContainer.cpp b/Engine/Plugins/Runtime/GameInput/Source/GameInputBase/Private/GameInputDeviceContainer.cpp
--- a/Engine/Plugins/Runtime/GameInput/Source/GameInputBase/Private/GameInputDeviceContainer.cpp
+++ b/Engine/Plugins/Runtime/GameInput/Source/GameInputBase/Private/GameInputDeviceContainer.cpp
@@ -164,7 +164,7 @@ void FGameInputDeviceContainer::InitalizeDeviceProcessors_Impl()
 			UE_LOG(
 				LogGameInput,
 				Warning,
-				TEXT("A raw game input device (%s) was connected but will not be processed because it does not have an explict device configuration"),
+				TEXT("A raw game input device (%s) was connected but will not be processed because it does not have an explicit device configuration"),
 				*UE::GameInput::LexToString(Device));
 		}
 	}
