[Game Input] Fix an issue with raw devices not working on GameInput versions <= 0 (the non-nuGet version).

We were incorrectly skipping the old way of reading raw reports, which lead to raw reports working on PC but not other targets.

#rnx
#Synced-CL 49643405

diff --git a/Engine/Plugins/Runtime/GameInput/Source/GameInputBase/Private/Processors/GameInputDeviceProcessor_Raw.cpp b/Engine/Plugins/Runtime/GameInput/Source/GameInputBase/Private/Processors/GameInputDeviceProcessor_Raw.cpp
--- a/Engine/Plugins/Runtime/GameInput/Source/GameInputBase/Private/Processors/GameInputDeviceProcessor_Raw.cpp
+++ b/Engine/Plugins/Runtime/GameInput/Source/GameInputBase/Private/Processors/GameInputDeviceProcessor_Raw.cpp
@@ -41,10 +41,15 @@ const float FGameInputRawDeviceProcessor::RawValueToFloatAnalog(uint8 RawValue,
 PRAGMA_DISABLE_DEPRECATION_WARNINGS
 const GameInputRawDeviceReportInfo* FGameInputRawDeviceProcessor::ReadCurrentRawInputState(const FGameInputEventParams& Params, IGameInputReading* ReadingToUse)
 {
-#if GAMEINPUT_API_VERSION >= 3
-	UE_LOG(LogGameInput, Error, TEXT("ReadCurrentRawInputState Is deprecated in Game Input version 3+. Call ReadCurrentRawInputStateInfo instead. (Device %s) "), *UE::GameInput::LexToString(Params.Device));
-	return nullptr;
-#else
+	return Internal_AttemptToReadCurrentRawInputState(Params, ReadingToUse);
+}
+PRAGMA_ENABLE_DEPRECATION_WARNINGS
+
+const GameInputRawDeviceReportInfo* FGameInputRawDeviceProcessor::Internal_AttemptToReadCurrentRawInputState(const FGameInputEventParams& Params, IGameInputReading* ReadingToUse)
+{
+	// 	The GetReportInfo function signature doesn't exist in game input version 1 or 2, and changed in 3.0.
+	// We still need this implementation for version 0 though, which is the GDK version of GameInput.
+#if GAMEINPUT_API_VERSION <= 0
 	
 	if (!ReadingToUse)
 	{
@@ -82,9 +87,12 @@ const GameInputRawDeviceReportInfo* FGameInputRawDeviceProcessor::ReadCurrentRaw
 	}
 
 	return RawReportInfo;
-#endif	// #if GAMEINPUT_API_VERSION >= 3
+	
+#else
+	ensureAlwaysMsgf(false, TEXT("Internal_AttemptToReadCurrentRawInputState should only be called for Game Input library version 0"));
+	return nullptr;
+#endif
 }
-PRAGMA_ENABLE_DEPRECATION_WARNINGS
 
 GameInputRawDeviceReportInfo FGameInputRawDeviceProcessor::ReadCurrentRawInputStateInfo(const FGameInputEventParams& Params, IGameInputReading* ReadingToUse)
 {
@@ -121,6 +129,16 @@ GameInputRawDeviceReportInfo FGameInputRawDeviceProcessor::ReadCurrentRawInputSt
 	const int32 NumReadBytes = RawReport->GetRawData(RawReport->GetRawDataSize(), CurrentRawData.GetData());
 	
 	RawReport->GetReportInfo(&RawReportInfo);
+	
+#elif GAMEINPUT_API_VERSION <= 0
+
+	// Older versions of the Game Input API at 0 or less can read raw data, but have different function signatures.
+	// We can attempt to populate our raw data based on this if we can read it
+	if (const GameInputRawDeviceReportInfo* InfoPtr = Internal_AttemptToReadCurrentRawInputState(Params, ReadingToUse))
+	{
+		RawReportInfo = *InfoPtr;
+	}
+
 #else
 	UE_LOG(LogGameInput, Error, TEXT("A GameInput version of 3+ is required for raw device support. (Device %s) "), *UE::GameInput::LexToString(Params.Device));
 #endif	// #if GAMEINPUT_API_VERSION >= 3
diff --git a/Engine/Plugins/Runtime/GameInput/Source/GameInputBase/Public/Processors/GameInputDeviceProcessor_Raw.h b/Engine/Plugins/Runtime/GameInput/Source/GameInputBase/Public/Processors/GameInputDeviceProcessor_Raw.h
--- a/Engine/Plugins/Runtime/GameInput/Source/GameInputBase/Public/Processors/GameInputDeviceProcessor_Raw.h
+++ b/Engine/Plugins/Runtime/GameInput/Source/GameInputBase/Public/Processors/GameInputDeviceProcessor_Raw.h
@@ -31,7 +31,18 @@ protected:
 	*/
 	UE_DEPRECATED(5.9, "Use ReadCurrentRawInputStateInfo instead.")
 	GAMEINPUTBASE_API const GameInputRawDeviceReportInfo* ReadCurrentRawInputState(const FGameInputEventParams& Params, IGameInputReading* ReadingToUse);
+
+private:
 	
+	/**
+	 * Reads the raw report info for GameInput API version 0. If the Api version is different then 0, this will return nullptr.	 
+	 */
+	const GameInputRawDeviceReportInfo* Internal_AttemptToReadCurrentRawInputState(const FGameInputEventParams& Params, IGameInputReading* ReadingToUse);
+
+protected:
+	/**
+	 * Returns the raw report info about the given reading which can be used to access the raw serialized data such as the size, report id, and more.	 
+	 */
 	GAMEINPUTBASE_API GameInputRawDeviceReportInfo ReadCurrentRawInputStateInfo(const FGameInputEventParams& Params, IGameInputReading* ReadingToUse);
 
 	/**
