[Game Input] Only log the number of readings if they are greater than zero. This just reduces pointless log spam when debugging devices.

#rnx
#Synced-CL 49643405

diff --git a/Engine/Plugins/Runtime/GameInput/Source/GameInputBase/Private/GameInputDeviceContainer.cpp b/Engine/Plugins/Runtime/GameInput/Source/GameInputBase/Private/GameInputDeviceContainer.cpp
--- a/Engine/Plugins/Runtime/GameInput/Source/GameInputBase/Private/GameInputDeviceContainer.cpp
+++ b/Engine/Plugins/Runtime/GameInput/Source/GameInputBase/Private/GameInputDeviceContainer.cpp
@@ -332,7 +332,7 @@ const GameInputKind FGameInputDeviceContainer::ProcessInput(IGameInput* GameInpu
 		LastReadingTimestamp = LastReading->GetTimestamp();
 	}
 
-	UE_LOG(LogGameInput, VeryVerbose, TEXT("Processed '%d' GameInput readings off the input stack for device: %s"), NumReadingsProcessed, *UE::GameInput::LexToString(Device));
+	UE_CLOG(NumReadingsProcessed > 0, LogGameInput, VeryVerbose, TEXT("Processed '%d' GameInput readings off the input stack for device: %s"), NumReadingsProcessed, *UE::GameInput::LexToString(Device));
 
 	return OutProcessedInputKinds;
 }
