[Game Input] Add a raw report option for treating a single uint8 index as two values, where the upper 4 bits represent one value and the lower 4 bits represent another. Using this method, also allow you to treat these reported values as an input chord (multiple keys all at once).

This is useful because the existing option of "Treating raw data as a bitmask" would only support digital keys, but often times devices report these packed float values.

#tests Test with an input device which supports this kind of reading (Rock Band drums + Cymbals) to ensure that the buttons are reported as expected.
#rnx

#Synced-CL 49989401

diff --git a/Engine/Plugins/Runtime/GameInput/Source/GameInputBase/Private/Processors/GameInputDeviceProcessor_Raw.cpp b/Engine/Plugins/Runtime/GameInput/Source/GameInputBase/Private/Processors/GameInputDeviceProcessor_Raw.cpp
--- a/Engine/Plugins/Runtime/GameInput/Source/GameInputBase/Private/Processors/GameInputDeviceProcessor_Raw.cpp
+++ b/Engine/Plugins/Runtime/GameInput/Source/GameInputBase/Private/Processors/GameInputDeviceProcessor_Raw.cpp
@@ -5,6 +5,7 @@
 #include "GameInputDeveloperSettings.h"
 #include "GameInputLogging.h"
 #include "GameInputUtils.h"
+#include "InputCoreTypes.h"
 
 #if GAME_INPUT_SUPPORT
 
@@ -16,6 +17,11 @@ constexpr float FGameInputRawDeviceProcessor::RawValueToFloatTrigger(const uint8
 	return (static_cast<float>(RawValue) / 255.f);
 }
 
+constexpr float FGameInputRawDeviceProcessor::RawValueNibbleValueToFloat(const uint8 RawValue, const float Scale) const
+{
+	return (static_cast<float>(RawValue) / Scale);
+}
+
 const float FGameInputRawDeviceProcessor::RawValueToFloatAnalog(uint8 RawValue, const uint8 DeadZone /* = 2 */) const
 {
 	// Apply a simple square deadzone...
@@ -243,7 +249,7 @@ bool FGameInputRawDeviceProcessor::PostProcessInput(const FGameInputEventParams&
 
 bool FGameInputRawDeviceProcessor::ProcessAllRawValues(const FGameInputEventParams& Params, const FGameInputDeviceConfiguration* DeviceConfig, const bool bShouldProcessButtons, const bool bShouldProcessAnalog)
 {
-	ensure(CurrentRawData.Num() == PreviousRawData.Num());
+	ensure(CurrentRawData.Num() <= PreviousRawData.Num());
 
 	ensure(Params.PlatformUserId.IsValid());
 
@@ -270,6 +276,13 @@ bool FGameInputRawDeviceProcessor::ProcessAllRawValues(const FGameInputEventPara
 					ProcessRawInputValueAsAanalogPaired(Params, AxisData);
 				}
 			}
+			else if (AxisData->TranslationBehavior == ERawDeviceReportTranslationBehavior::TreatAsPackedAxisUint8)
+			{
+				if (AxisData->bTreatValuesAsAnalog && bShouldProcessAnalog || (!AxisData->bTreatValuesAsAnalog && bShouldProcessButtons))
+				{
+					bRes |= ProcessRawInputValuePackedUInt8(Params, i, AxisData);	
+				}
+			}
 			// All other methods require a valid key name on the config
 			else if (AxisData->KeyName.IsValid())
 			{
@@ -378,9 +391,11 @@ bool FGameInputRawDeviceProcessor::ProcessRawInputValueAsAanalog(const FGameInpu
 
 	const float CurrentValueFloat = AxisData->Scalar * (AxisData->TranslationBehavior == ERawDeviceReportTranslationBehavior::TreatAsTrigger ? RawValueToFloatTrigger(Val) : RawValueToFloatAnalog(Val, AxisData->AnalogDeadzone));
 	const float PreviousValueFloat = AxisData->Scalar * (AxisData->TranslationBehavior == ERawDeviceReportTranslationBehavior::TreatAsTrigger ? RawValueToFloatTrigger(PrevVal) : RawValueToFloatAnalog(PrevVal, AxisData->AnalogDeadzone));
-
-	if (!AxisData->KeyName.IsValid())
+	
+	// A valid FKey is required
+	if (!FKey(AxisData->KeyName).IsValid())
 	{
+		UE_LOG(LogGameInput, Error, TEXT("Attempting to report an invalid key name '%s' (Device %s) "), *AxisData->KeyName.ToString(), *UE::GameInput::LexToString(Params.Device));
 		return false;
 	}
 
@@ -432,6 +447,169 @@ bool FGameInputRawDeviceProcessor::ProcessRawInputValueAsAanalogPaired(const FGa
 	return CurrentValueFloat != 0.0f;
 }
 
+bool FGameInputRawDeviceProcessor::ProcessRawInputValuePackedUInt8(const FGameInputEventParams& Params, const int32 RawValueIndex, const FGameInputRawDeviceReportData* AxisData)
+{
+	check(AxisData && AxisData->TranslationBehavior == ERawDeviceReportTranslationBehavior::TreatAsPackedAxisUint8);
+
+	// Isolate the lower bits as a uint8 value (masking off the upper 4, and only using the lower 4)
+	const uint8 CurrentLowerBitValues = CurrentRawData[RawValueIndex] & 0x0F;
+	const uint8 PreviousLowerBitValues = PreviousRawData[RawValueIndex] & 0x0F;
+	
+	// To get the upper 4 bits as a uint8 value, we can just shift them right 4 and drop off the lower 4 values.
+	const uint8 CurrentUpperBitValues = CurrentRawData[RawValueIndex] >> 4;
+	const uint8 PreviousUpperBitValues = PreviousRawData[RawValueIndex] >> 4;
+	
+	// Map the value of the key name to this 1 if we haven't already, to keep track of the button repeat timings
+	FPerRawInputIndexData& IndexData = RawInputIndexDataMap.FindOrAdd(RawValueIndex);
+
+	// Report these values as analog data
+	if (AxisData->bTreatValuesAsAnalog)
+	{
+		EvaluatePackedUInt8ButtonChordsAsAnalog(Params, CurrentLowerBitValues, PreviousLowerBitValues, AxisData->LowerBitKeys);
+    	EvaluatePackedUInt8ButtonChordsAsAnalog(Params, CurrentUpperBitValues, PreviousUpperBitValues, AxisData->UpperBitKeys);	
+	}
+	// Otherwise, report them as key press events
+	else
+	{
+		static constexpr int32 LowerBitsRepeatTimeIndex = 0;
+		static constexpr int32 UpperBitsRepeatTimeIndex = 1;
+	
+
+		// Notice: We are using the !! operator here because we want to only send values of 0 or 1 for button presses
+		// to avoid needlessly over reporting. For non-analog data, anything that is non-zero will be considered a "pressed" state
+		EvaluatePackedUInt8ButtonChordsAsButtonEvents(Params, !!CurrentLowerBitValues, !!PreviousLowerBitValues, &IndexData.RepeatTime[LowerBitsRepeatTimeIndex], AxisData->LowerBitKeys);
+		EvaluatePackedUInt8ButtonChordsAsButtonEvents(Params, !!CurrentUpperBitValues, !!PreviousUpperBitValues, &IndexData.RepeatTime[UpperBitsRepeatTimeIndex], AxisData->UpperBitKeys);
+	}
+
+	// If the upper or lower value is non-zero then we had a reading
+	return CurrentLowerBitValues != 0 || CurrentUpperBitValues != 0;
+}
+
+void FGameInputRawDeviceProcessor::EvaluatePackedUInt8ButtonChordsAsAnalog(
+	const FGameInputEventParams& Params,
+	const uint8 CurrentValue,
+	const uint8 PreviousValue,
+	const TArray<FGamepadKeyNames::Type>& ButtonNames)
+{
+	// We can't tell slate about input messages from an invalid platform user
+	if (!Params.PlatformUserId.IsValid())
+	{
+		UE_LOG(LogGameInput, Verbose, TEXT("Attempting to evaluate button states with an invalid platform user id of '%d'. The button messages will not be sent."), Params.PlatformUserId.GetInternalId());
+		return;
+	}
+
+	IPlatformInputDeviceMapper& DeviceMapper = IPlatformInputDeviceMapper::Get();
+
+	FInputDeviceScope InputScope(nullptr, FName("GameInput"), DeviceMapper.GetUserIndexForPlatformUser(Params.PlatformUserId), GetHardwareDeviceIdentifierName(Params));
+	
+	if (CurrentValue != PreviousValue)
+	{
+		const float CurrentValueFloat = RawValueNibbleValueToFloat(CurrentValue);
+
+		// Report an analog value for each of the buttons that we want
+		for (const FGamepadKeyNames::Type& KeyName : ButtonNames)
+		{
+			if (!FKey(KeyName).IsValid())
+			{
+				continue;
+			}
+			
+			UE_LOG(LogGameInput, Verbose, TEXT("[FGameInputRawDeviceProcessor::EvaluatePackedUInt8ButtonChordsAsAnalog] (PlatformUserId = %d, InputDeviceId = %d) - Button '%s' Value %.3f"),
+				Params.PlatformUserId.GetInternalId(), Params.InputDeviceId.GetId(), *KeyName.ToString(), CurrentValueFloat);
+			
+			Params.MessageHandler->OnControllerAnalog(KeyName, Params.PlatformUserId, Params.InputDeviceId, CurrentValueFloat);	
+		}
+	}
+}
+
+void FGameInputRawDeviceProcessor::EvaluatePackedUInt8ButtonChordsAsButtonEvents(
+	const FGameInputEventParams& Params,
+	const uint8 CurrentValue,
+	const uint8 PreviousValue,
+	double* RepeatTime,
+	const TArray<FGamepadKeyNames::Type>& ButtonNames)
+{
+	// We can't tell slate about input messages from an invalid platform user
+	if (!Params.PlatformUserId.IsValid())
+	{
+		UE_LOG(LogGameInput, Verbose, TEXT("Attempting to evaluate button states with an invalid platform user id of '%d'. The button messages will not be sent."), Params.PlatformUserId.GetInternalId());
+		return;
+	}
+
+	if (!RepeatTime)
+	{
+		UE_LOG(LogGameInput, Error, TEXT("[IGameInputDeviceProcessor::EvaluateButtonStates] Invalid RepeatTime array given to evaluate button states!"));
+		return;
+	}
+
+	if (ButtonNames.IsEmpty())
+	{
+		return;
+	}
+
+	IPlatformInputDeviceMapper& DeviceMapper = IPlatformInputDeviceMapper::Get();
+
+	const FInputDeviceScope InputScope(nullptr, FName("GameInput"), DeviceMapper.GetUserIndexForPlatformUser(Params.PlatformUserId), GetHardwareDeviceIdentifierName(Params));
+
+	auto ForEachValidKey = [&ButtonNames](TFunction<void(const FGamepadKeyNames::Type&)>&& Predicate)
+	{
+		for (const FGamepadKeyNames::Type& KeyName : ButtonNames)
+		{
+			if (!FKey(KeyName).IsValid())
+			{
+				continue;
+			}
+			
+			Predicate(KeyName);
+		}
+	};
+
+	const double CurrentTime = FPlatformTime::Seconds();
+
+	// Handle the pressing or releasing of a key (the value has changed)
+	if (CurrentValue != PreviousValue)
+	{
+		if (CurrentValue)
+		{
+			ForEachValidKey([&Params](const FGamepadKeyNames::Type& KeyName)
+			{
+				UE_LOG(LogGameInput, Verbose, TEXT("[FGameInputDevice::EvaluateButtonStates] (PlatformUserId = %d, InputDeviceId = %d) - Button '%s' Pressed"),
+					Params.PlatformUserId.GetInternalId(), Params.InputDeviceId.GetId(), *KeyName.ToString());
+					
+				Params.MessageHandler->OnControllerButtonPressed(KeyName, Params.PlatformUserId, Params.InputDeviceId, false);
+			});
+
+			*RepeatTime = CurrentTime + InitialButtonRepeatDelay;
+		}
+		else
+		{
+			ForEachValidKey([&Params](const FGamepadKeyNames::Type& KeyName)
+			{
+				UE_LOG(LogGameInput, Verbose, TEXT("[FGameInputDevice::EvaluateButtonStates] (PlatformUserId = %d, InputDeviceId = %d) - Button '%s' Released"),
+					Params.PlatformUserId.GetInternalId(), Params.InputDeviceId.GetId(), *KeyName.ToString());
+					
+				Params.MessageHandler->OnControllerButtonReleased(KeyName, Params.PlatformUserId, Params.InputDeviceId, false);
+			});
+
+			*RepeatTime = 0.0;
+		}
+	}
+
+	// Handle button repeat events
+	if (CurrentValue && CurrentValue == PreviousValue && CurrentTime > *RepeatTime)
+	{
+		*RepeatTime = CurrentTime + ButtonRepeatDelay;
+		
+		ForEachValidKey([&Params](const FGamepadKeyNames::Type& KeyName)
+		{
+			UE_LOG(LogGameInput, Verbose, TEXT("[FGameInputDevice::EvaluateButtonStates] (PlatformUserId = %d, InputDeviceId = %d) - Button '%s' Repeat Pressed"),
+				Params.PlatformUserId.GetInternalId(), Params.InputDeviceId.GetId(), *KeyName.ToString());
+				
+			Params.MessageHandler->OnControllerButtonPressed(KeyName, Params.PlatformUserId, Params.InputDeviceId, true);
+		});
+	}
+}
+
 void FGameInputRawDeviceProcessor::ClearState(const FGameInputEventParams& Params)
 {
 	// Can't do anything for an invalid platform user
diff --git a/Engine/Plugins/Runtime/GameInput/Source/GameInputBase/Public/GameInputDeveloperSettings.h b/Engine/Plugins/Runtime/GameInput/Source/GameInputBase/Public/GameInputDeveloperSettings.h
--- a/Engine/Plugins/Runtime/GameInput/Source/GameInputBase/Public/GameInputDeveloperSettings.h
+++ b/Engine/Plugins/Runtime/GameInput/Source/GameInputBase/Public/GameInputDeveloperSettings.h
@@ -109,6 +109,10 @@ enum class ERawDeviceReportTranslationBehavior : uint8
 	// Treat this as two raw values, where one axis is the "lower" bits of an int16 and 
 	// one is the "higher" bits of the int16. 
 	TreatAsPackedAxisPair = 4,
+
+	// Treat this raw report index as two values, reporting two different sets of buttons.
+	// The lower 4 bits represent one set of buttons, and the upper 4 bits represent another set of buttons.
+	TreatAsPackedAxisUint8 = 5,
 };
 
 /**
@@ -131,7 +135,7 @@ struct FGameInputRawDeviceReportData
 	 * 
 	 * Note: This value only matters for Translation Behaviors other then the button bitmask.
 	 */
-	UPROPERTY(EditAnywhere, Config, Category = "Game Input Device", meta=(EditCondition="TranslationBehavior != ERawDeviceReportTranslationBehavior::TreatAsButtonBitmask", EditConditionHides))
+	UPROPERTY(EditAnywhere, Config, Category = "Game Input Device", meta=(EditCondition="TranslationBehavior != ERawDeviceReportTranslationBehavior::TreatAsButtonBitmask && TranslationBehavior != ERawDeviceReportTranslationBehavior::TreatAsPackedAxisUint8", EditConditionHides))
 	FName KeyName = NAME_None;
 
 	/**
@@ -174,12 +178,31 @@ struct FGameInputRawDeviceReportData
 	UPROPERTY(EditAnywhere, Config, Category = "Game Input Device", meta=(EditCondition="TranslationBehavior == ERawDeviceReportTranslationBehavior::TreatAsButtonBitmask", EditConditionHides))
 	TMap<int32, FName> ButtonBitMaskMappings;
 
-
 	UPROPERTY(EditAnywhere, Config, Category = "Game Input Device", meta = (EditCondition = "TranslationBehavior == ERawDeviceReportTranslationBehavior::TreatAsPackedAxisPair", EditConditionHides))
 	uint8 LowerBitAxisIndex = 0;
 	
 	UPROPERTY(EditAnywhere, Config, Category = "Game Input Device", meta = (EditCondition = "TranslationBehavior == ERawDeviceReportTranslationBehavior::TreatAsPackedAxisPair", EditConditionHides))
 	uint8 HigherBitAxisIndex = 0;
+
+	/**
+	* Array of FKey names to report when the lower bits of this raw value index are detected. Each element of this array will be reported once per-frame, known as an "input chord"
+	*/
+	UPROPERTY(EditAnywhere, Config, Category = "Game Input Device", meta = (EditCondition = "TranslationBehavior == ERawDeviceReportTranslationBehavior::TreatAsPackedAxisUint8", EditConditionHides))
+	TArray<FName> LowerBitKeys;
+
+	/**
+	* Array of FKey names to report when the upper bits of this raw value index are detected. Each element of this array will be reported once per-frame, known as an "input chord" 
+	*/
+	UPROPERTY(EditAnywhere, Config, Category = "Game Input Device", meta = (EditCondition = "TranslationBehavior == ERawDeviceReportTranslationBehavior::TreatAsPackedAxisUint8", EditConditionHides))
+	TArray<FName> UpperBitKeys;
+
+	/**
+	 * If true, then this packed bitmask value will be treated as an analog button.
+	 * This means that it will use "OnControllerAnalog", reporting a float value, instead of "OnControllerButtonPressed" using input events. 
+	 */
+	UPROPERTY(EditAnywhere, Config, Category = "Game Input Device", meta = (EditCondition = "TranslationBehavior == ERawDeviceReportTranslationBehavior::TreatAsPackedAxisUint8", EditConditionHides))
+	bool bTreatValuesAsAnalog = false;
+
 };
 
 /**
diff --git a/Engine/Plugins/Runtime/GameInput/Source/GameInputBase/Public/Processors/GameInputDeviceProcessor_Raw.h b/Engine/Plugins/Runtime/GameInput/Source/GameInputBase/Public/Processors/GameInputDeviceProcessor_Raw.h
--- a/Engine/Plugins/Runtime/GameInput/Source/GameInputBase/Public/Processors/GameInputDeviceProcessor_Raw.h
+++ b/Engine/Plugins/Runtime/GameInput/Source/GameInputBase/Public/Processors/GameInputDeviceProcessor_Raw.h
@@ -54,6 +54,11 @@ protected:
 	*/
 	GAMEINPUTBASE_API constexpr float RawValueToFloatTrigger(const uint8 RawValue) const;
 
+	/**
+	 * Converts the lower 4 bits of the given raw value to a value between 0.0f and +1.0f based on the given scale.
+	 */
+	GAMEINPUTBASE_API constexpr float RawValueNibbleValueToFloat(const uint8 RawValue, const float Scale = 15.0f) const;
+
 	/**
 	* Translates the given uint8 RawValue to a float between -1.0 and +1.0.
 	* 
@@ -93,6 +98,21 @@ protected:
 	*/
 	GAMEINPUTBASE_API bool ProcessRawInputValueAsAanalogPaired(const FGameInputEventParams& Params, const FGameInputRawDeviceReportData* AxisData);
 
+	GAMEINPUTBASE_API bool ProcessRawInputValuePackedUInt8(const FGameInputEventParams& Params, const int32 RawValueIndex, const FGameInputRawDeviceReportData* AxisData);
+
+	void EvaluatePackedUInt8ButtonChordsAsAnalog(
+			const FGameInputEventParams& Params,
+			const uint8 CurrentValue,
+			const uint8 PreviousValue,
+			const TArray<FGamepadKeyNames::Type>& ButtonNames);
+
+	void EvaluatePackedUInt8ButtonChordsAsButtonEvents(
+			const FGameInputEventParams& Params,
+			const uint8 CurrentValue,
+			const uint8 PreviousValue,
+			double* RepeatTime,
+			const TArray<FGamepadKeyNames::Type>& ButtonNames);
+
 	/**
 	* The current raw input data that is populated by IGameInputRawDeviceReport::GetRawData
 	*/
