[Game Input] Make it possible to ignore or filter out devices from game input based on a VID/PID combo or their device family.

#tests Editor win64 target of lyra using game input, toggling devices on and off. Run cooked build on console target to ensure functionality.
[RN] Make it possible to ignore or filter out devices from game input based on a VID/PID combo or their device family.

#Synced-CL 50735006

diff --git a/Engine/Plugins/GameInputWindows/Source/GameInputWindows/Private/GameInputWindowsDevice.cpp b/Engine/Plugins/GameInputWindows/Source/GameInputWindows/Private/GameInputWindowsDevice.cpp
--- a/Engine/Plugins/GameInputWindows/Source/GameInputWindows/Private/GameInputWindowsDevice.cpp
+++ b/Engine/Plugins/GameInputWindows/Source/GameInputWindows/Private/GameInputWindowsDevice.cpp
@@ -198,8 +198,20 @@ void FGameInputWindowsInputDevice::HandleEditorSettingChanged()
 {
 	for (TSharedPtr<FGameInputDeviceContainer> Device : DeviceData)
 	{
+		GameInputKind SupportedKind = GetCurrentGameInputKindSupport();
+
+		// Make sure that this device is still supported. If you are in the editor then you may have changed
+		// this value and a connected device is no longer supported. If that is the case, then make sure we don't create
+		// any new processors for it at all.
+		IGameInputDevice* GameInputDevice = Device->GetGameInputDevice();
+		if (!IsSupportedDevice(GameInputDevice))
+		{
+			UE_LOG(LogGameInput, Log, TEXT("'%s' is not supported (the settings have changed) and input will no longer be processed for it."), *UE::GameInput::LexToString(GameInputDevice));
+			SupportedKind = GameInputKind::GameInputKindUnknown;
+		}
+		
 		// Recreate the device processors for every device
-		Device->RecreateDeviceProcessors(GetCurrentGameInputKindSupport());
+		Device->RecreateDeviceProcessors(SupportedKind);
 
 		// If this device previously had an invalid DeviceId because it had no processors, but now it does,
 		// then we need to assign a FInputDeviceId because it can now process input
diff --git a/Engine/Plugins/Runtime/GameInput/Source/GameInputBase/Private/GameInputDeveloperSettings.cpp b/Engine/Plugins/Runtime/GameInput/Source/GameInputBase/Private/GameInputDeveloperSettings.cpp
--- a/Engine/Plugins/Runtime/GameInput/Source/GameInputBase/Private/GameInputDeveloperSettings.cpp
+++ b/Engine/Plugins/Runtime/GameInput/Source/GameInputBase/Private/GameInputDeveloperSettings.cpp
@@ -84,6 +84,61 @@ const FGameInputDeviceConfiguration* UGameInputDeveloperSettings::FindDeviceConf
 	});
 }
 
+namespace UE::GameInput
+{
+	/** Converts the "raw" game input device family to the UE friendly property version. */
+	static EGameInputDeviceFamily ConvertDeviceFamily(const GameInputDeviceFamily RawFamily)
+	{
+		switch (RawFamily)
+		{
+		case GameInputFamilyVirtual:
+			return EGameInputDeviceFamily::Virtual;
+#if GAMEINPUT_API_VERSION >= 3		
+		case GameInputFamilyUnknown:
+			return EGameInputDeviceFamily::Unknown;
+#endif			
+		case GameInputFamilyXboxOne:
+			return EGameInputDeviceFamily::XboxOne;
+		case GameInputFamilyXbox360:
+			return EGameInputDeviceFamily::Xbox360;
+		case GameInputFamilyHid:
+			return EGameInputDeviceFamily::Hid;
+		case GameInputFamilyI8042:
+			return EGameInputDeviceFamily::I8042;
+		case GameInputFamilyAggregate:
+			return EGameInputDeviceFamily::Aggregate;
+		}
+		
+		return EGameInputDeviceFamily::None;
+	}
+}
+
+bool UGameInputDeveloperSettings::ShouldIgnoreDevice(const GameInputDeviceInfo* const Info) const
+{
+	// ignore any null devices
+	if (!Info)
+	{
+		return true;
+	}
+
+	// We can ignore based on VID/PID...
+	const FGameInputDeviceIdentifier DeviceId(Info->vendorId, Info->productId);
+	if (DevicesToIgnore.Contains(DeviceId))
+	{
+		return true;
+	}
+
+	// We can also ignore based on the broader device family if desired
+	const EGameInputDeviceFamily DeviceFamily = UE::GameInput::ConvertDeviceFamily(Info->deviceFamily);
+	if (!EnumHasAllFlags(static_cast<EGameInputDeviceFamily>(SupportedDeviceFamilies), DeviceFamily))
+	{
+		return true;
+	}
+
+	// We fully support this device
+	return false;
+}
+
 #endif	// GAME_INPUT_SUPPORT
 
 UGameInputPlatformSettings* UGameInputPlatformSettings::Get()
@@ -115,6 +170,7 @@ void UGameInputPlatformSettings::InitializePlatformDefaults()
 
 UGameInputDeveloperSettings::UGameInputDeveloperSettings(const FObjectInitializer& Initializer)
 	: Super(Initializer)
+	, SupportedDeviceFamilies(static_cast<uint32>(EGameInputDeviceFamily::All))
 	, bDoNotProcessDuplicateCapabilitiesForSingleUser(true)
 {
 	PlatformSpecificSettings.Initialize(UGameInputPlatformSettings::StaticClass());
diff --git a/Engine/Plugins/Runtime/GameInput/Source/GameInputBase/Private/IGameInputDeviceInterface.cpp b/Engine/Plugins/Runtime/GameInput/Source/GameInputBase/Private/IGameInputDeviceInterface.cpp
--- a/Engine/Plugins/Runtime/GameInput/Source/GameInputBase/Private/IGameInputDeviceInterface.cpp
+++ b/Engine/Plugins/Runtime/GameInput/Source/GameInputBase/Private/IGameInputDeviceInterface.cpp
@@ -389,9 +389,20 @@ bool IGameInputDeviceInterface::BindDeviceStatusCallbacks()
 }
 
 void IGameInputDeviceInterface::OnDeviceConnectionChanged(GameInputCallbackToken CallbackToken, IGameInputDevice* Device, uint64 Timestamp, GameInputDeviceStatus CurrentStatus, GameInputDeviceStatus PreviousStatus)
-{
+{	
+	if (!IsSupportedDevice(Device))
+	{
+		UE_LOG(
+			LogGameInput,
+			Verbose,
+			TEXT("%s is an unsupported Device and will not be processed. (Connection Changed detected from '%s' to '%s')"),
+			*UE::GameInput::LexToString(Device), *UE::GameInput::LexToString(PreviousStatus), *UE::GameInput::LexToString(CurrentStatus));
+		
+		return;
+	}
+	
 	FScopeLock Lock(&DeviceInfoCS);
-
+	
 	UE_LOG(LogGameInput, Log, TEXT("Device Connection Changed from '%s' to '%s'"), *UE::GameInput::LexToString(PreviousStatus), *UE::GameInput::LexToString(CurrentStatus));
 	
 	// This event may come in async from GameInput from outside the game thread, so defer it until Tick so we can guarantee game thread access
@@ -405,6 +416,30 @@ void IGameInputDeviceInterface::OnDeviceConnectionChanged(GameInputCallbackToken
 	DeferredDeviceConnectionChanges.Emplace(Event);
 }
 
+bool IGameInputDeviceInterface::IsSupportedDevice(IGameInputDevice* Device) const
+{
+	if (!Device)
+	{
+		return false;
+	}
+	
+	const GameInputDeviceInfo* Info = UE::GameInput::GetDeviceInfo(Device);
+	
+	if (!Info)
+	{
+		return false;
+	}
+
+	const UGameInputDeveloperSettings* Settings = GetDefault<UGameInputDeveloperSettings>();
+
+	if (Settings && Settings->ShouldIgnoreDevice(Info))
+	{
+		return false;
+	}
+	
+	return true;
+}
+
 void IGameInputDeviceInterface::ProcessDeferredDeviceConnectionChanges()
 {
 	// We only want to actually handle device connection events in the game thread
diff --git a/Engine/Plugins/Runtime/GameInput/Source/GameInputBase/Public/GameInputDeveloperSettings.h b/Engine/Plugins/Runtime/GameInput/Source/GameInputBase/Public/GameInputDeveloperSettings.h
--- a/Engine/Plugins/Runtime/GameInput/Source/GameInputBase/Public/GameInputDeveloperSettings.h
+++ b/Engine/Plugins/Runtime/GameInput/Source/GameInputBase/Public/GameInputDeveloperSettings.h
@@ -52,6 +52,25 @@ struct FGameInputDeviceIdentifier
 	UE_API FString ToString() const;
 };
 
+/**
+* UPROP friendly representation of GameInput's GameInputDeviceFamily type.
+*/
+UENUM(meta = (BitFlags, UseEnumValuesAsMaskValuesInEditor = "true"))
+enum class EGameInputDeviceFamily : uint32
+{
+	None 	  =  0 UMETA(Hidden),
+	Virtual   =  1 << 0,
+	Unknown   =  1 << 1,
+	XboxOne   =  1 << 2,
+	Xbox360   =  1 << 3,
+	Hid       =  1 << 4,
+	I8042     =  1 << 5,
+	Aggregate =  1 << 6,
+	All		  =  Virtual | Unknown | XboxOne | Xbox360 | Hid | I8042 | Aggregate UMETA(Hidden),
+};
+
+ENUM_CLASS_FLAGS(EGameInputDeviceFamily)
+
 /**
 * Configuration of an individual axis from a Game Input Controller device type.
 * 
@@ -478,6 +497,13 @@ public:
 	 * Returns a pointer to the device configuration that may exist for the given hardware ID
 	 */
 	UE_API const FGameInputDeviceConfiguration* FindDeviceConfiguration(const FGameInputDeviceIdentifier& HardwareID) const;
+
+	/**
+	* Returns true if the gien device should be completely ignored by GameInput processing.
+	* If a device is ignored, then it will not even fire connection event delegates.
+	*/
+	bool ShouldIgnoreDevice(const GameInputDeviceInfo* const Info) const; 
+	
 #endif	 // GAME_INPUT_SUPPORT
 
 	UE_API UGameInputDeveloperSettings(const FObjectInitializer& Initializer);
@@ -506,6 +532,22 @@ protected:
 	UPROPERTY(EditAnywhere, Category = "Game Input|Platform Options")
 	FPerPlatformSettings PlatformSpecificSettings;
 
+	/**
+	 * A list of devices for game input to completely ignore. If a device is connected that is in this list,
+	 * no processors will ever be created for it at all.
+	 */
+	UPROPERTY(EditAnywhere, Config, Category = "Device Settings")
+	TArray<FGameInputDeviceIdentifier> DevicesToIgnore;
+
+	/**
+	 * A bitmask of which GameInput device families you want to support. 
+	 * If a family is not enabled as a part of this list, it will be ignored.
+	 * 
+	 * Default: All families are supported.
+	 */
+	UPROPERTY(EditAnywhere, Config, Category = "Device Settings", meta = (Bitmask, BitmaskEnum = "/Script/GameInputBase.EGameInputDeviceFamily"))
+	uint32 SupportedDeviceFamilies;
+
 public:
 	/**
 	* If true, when you have multiple input devices mapped to a single FPlatformUserId,
diff --git a/Engine/Plugins/Runtime/GameInput/Source/GameInputBase/Public/IGameInputDeviceInterface.h b/Engine/Plugins/Runtime/GameInput/Source/GameInputBase/Public/IGameInputDeviceInterface.h
--- a/Engine/Plugins/Runtime/GameInput/Source/GameInputBase/Public/IGameInputDeviceInterface.h
+++ b/Engine/Plugins/Runtime/GameInput/Source/GameInputBase/Public/IGameInputDeviceInterface.h
@@ -54,6 +54,14 @@ protected:
 	 */
 	GAMEINPUTBASE_API virtual void OnDeviceConnectionChanged(GameInputCallbackToken CallbackToken, IGameInputDevice* Device, uint64 Timestamp, GameInputDeviceStatus CurrentStatus, GameInputDeviceStatus PreviousStatus);
 
+	/**
+	 * Returns true if the given IGameInputDevice is one which is allowed to be processed based on the Game Input Developer Settings.
+	 * 
+	 * This may return false for devices which have specific PID/VID's which this project does not wish GameInput to process,
+	 * or if the device is within a specific family that we do not want to process.
+	 */
+	GAMEINPUTBASE_API virtual bool IsSupportedDevice(IGameInputDevice* Device) const;
+
 	/**
 	 * Mark this device as disconnected and invalidate the tracked game input device. A disconnected device will stay in the
 	 * GameInputSourceDevices with a null Device pointer.
