[Game Input] Improve the description of the Game Input developer settings to better explain what it is.

#rnx

#Synced-CL 50969376

diff --git a/Engine/Plugins/Runtime/GameInput/Source/GameInputBaseEditor/Private/GameInputBaseEditorModule.cpp b/Engine/Plugins/Runtime/GameInput/Source/GameInputBaseEditor/Private/GameInputBaseEditorModule.cpp
--- a/Engine/Plugins/Runtime/GameInput/Source/GameInputBaseEditor/Private/GameInputBaseEditorModule.cpp
+++ b/Engine/Plugins/Runtime/GameInput/Source/GameInputBaseEditor/Private/GameInputBaseEditorModule.cpp
@@ -54,7 +54,10 @@ private:
 		{
 			SettingsModule->RegisterSettings("Project", "Plugins", "Input",
 				LOCTEXT("ControllerSettingsName", "Game Input Device Settings"),
-				LOCTEXT("ControllerSettingsDescription", "Settings for the Game Input plugin."),
+				LOCTEXT("ControllerSettingsDescription", "Settings for the Game Input plugin.\n\n"
+								   "\"Game Input\" is an API that interfaces with input devices of all kinds.\n"
+								   "This replaces the functionality of the older XInput and RawInput interfaces in Unreal, as well as some of the WinDualShock plugin.\n"
+								   "The minimum Windows version required for Game Input is Windows 10 19H1 (May 2019 update)."),
 				GetMutableDefault<UGameInputDeveloperSettings>()
 			);
 		}
