[Game Input] Make it easier to access the info required to build a FInputDeviceScope

#Synced-CL 50969376

diff --git a/Engine/Plugins/GameInputWindows/Source/GameInputWindows/Private/GameInputWindowsDevice.cpp b/Engine/Plugins/GameInputWindows/Source/GameInputWindows/Private/GameInputWindowsDevice.cpp
--- a/Engine/Plugins/GameInputWindows/Source/GameInputWindows/Private/GameInputWindowsDevice.cpp
+++ b/Engine/Plugins/GameInputWindows/Source/GameInputWindows/Private/GameInputWindowsDevice.cpp
@@ -115,14 +115,20 @@ void FGameInputWindowsInputDevice::HandleDeviceConnected(IGameInputDevice* Devic
 	const FInputDeviceId DeviceId = Data->GetDeviceId();
 	
 	// We only want to map this input device to a user if it has a valid Input Device Id.
-	// Everything has an input device id, so if we got one that is invalid
+	// Everything has an input device id, so if we got one that is invalid 
 	// then that means this device has no processors in it because they were explicitly disabled.
 	if (DeviceId.IsValid())
 	{
 		IPlatformInputDeviceMapper& DeviceMapper = IPlatformInputDeviceMapper::Get();
 		
 		const FPlatformUserId UserToAssign = DeviceMapper.GetPlatformUserForNewlyConnectedDevice();
-	
+		
+		FInputDeviceScope InputScope(
+			nullptr,
+			UE::GameInput::GetDeviceScopeClassName(),
+			DeviceMapper.GetUserIndexForPlatformUser(UserToAssign),
+			FString(UE::GameInput::GetHardwareDeviceIdentifierName(Device)));
+		
 		const bool bSuccess = DeviceMapper.Internal_MapInputDeviceToUser(DeviceId, UserToAssign, EInputDeviceConnectionState::Connected);
 		if (ensure(bSuccess))
 		{
diff --git a/Engine/Plugins/Runtime/GameInput/Source/GameInputBase/Private/GameInputUtils.cpp b/Engine/Plugins/Runtime/GameInput/Source/GameInputBase/Private/GameInputUtils.cpp
--- a/Engine/Plugins/Runtime/GameInput/Source/GameInputBase/Private/GameInputUtils.cpp
+++ b/Engine/Plugins/Runtime/GameInput/Source/GameInputBase/Private/GameInputUtils.cpp
@@ -372,6 +372,60 @@ namespace UE::GameInput
 
 		return Info;
 	}
+
+	FName GetDeviceScopeClassName()
+	{
+		static const FName ScopeClassName(TEXT("GameInput"));
+		return ScopeClassName;
+	}
+
+	FStringView GetHardwareDeviceIdentifierName(IGameInputDevice* Device)
+	{
+		static const TCHAR* ID_Virtual = TEXT("Virtual");
+		static const TCHAR* ID_Aggregate = TEXT("Aggregate");
+		static const TCHAR* ID_XboxOne = TEXT("XboxOne");
+		static const TCHAR* ID_Xbox360 = TEXT("Xbox360");
+		static const TCHAR* ID_Hid = TEXT("Hid");
+		static const TCHAR* ID_I8042 = TEXT("I8042");
+
+		if (Device)
+		{
+			if (const GameInputDeviceInfo* Info = UE::GameInput::GetDeviceInfo(Device))
+			{
+				// Check for any device specific overrides that may be there for custom devices
+				if (const UGameInputDeveloperSettings* Settings = GetDefault<UGameInputDeveloperSettings>())
+				{
+					if (const FGameInputDeviceConfiguration* Config = Settings->FindDeviceConfiguration(Info))
+					{
+						if (Config->bOverrideHardwareDeviceIdString)
+						{
+							return *Config->OverriddenHardwareDeviceId;
+						}					
+					}
+				}
+
+				switch (Info->deviceFamily)
+				{
+				case GameInputFamilyVirtual:
+					return ID_Virtual;
+				case GameInputFamilyAggregate:
+					return ID_Aggregate;
+				case GameInputFamilyXbox360:
+					return ID_Xbox360;
+				case GameInputFamilyHid:
+					return ID_Hid;
+				case GameInputFamilyI8042:
+					return ID_I8042;
+				case GameInputFamilyXboxOne:
+				default:
+					return ID_XboxOne;
+				}
+			}
+		}
+				
+		// If for some reason we are given a null device, default to just "Xbox One"
+		return ID_XboxOne;
+	}
 	
 }
 
diff --git a/Engine/Plugins/Runtime/GameInput/Source/GameInputBase/Private/Processors/GameInputDeviceProcessor_Raw.cpp b/Engine/Plugins/Runtime/GameInput/Source/GameInputBase/Private/Processors/GameInputDeviceProcessor_Raw.cpp
--- a/Engine/Plugins/Runtime/GameInput/Source/GameInputBase/Private/Processors/GameInputDeviceProcessor_Raw.cpp
+++ b/Engine/Plugins/Runtime/GameInput/Source/GameInputBase/Private/Processors/GameInputDeviceProcessor_Raw.cpp
@@ -500,7 +500,11 @@ void FGameInputRawDeviceProcessor::EvaluatePackedUInt8ButtonChordsAsAnalog(
 
 	IPlatformInputDeviceMapper& DeviceMapper = IPlatformInputDeviceMapper::Get();
 
-	FInputDeviceScope InputScope(nullptr, FName("GameInput"), DeviceMapper.GetUserIndexForPlatformUser(Params.PlatformUserId), GetHardwareDeviceIdentifierName(Params));
+	FInputDeviceScope InputScope(
+		nullptr,
+		UE::GameInput::GetDeviceScopeClassName(),
+		DeviceMapper.GetUserIndexForPlatformUser(Params.PlatformUserId),
+		FString(UE::GameInput::GetHardwareDeviceIdentifierName(Params.Device)));
 	
 	if (CurrentValue != PreviousValue)
 	{
@@ -549,7 +553,7 @@ void FGameInputRawDeviceProcessor::EvaluatePackedUInt8ButtonChordsAsButtonEvents
 
 	IPlatformInputDeviceMapper& DeviceMapper = IPlatformInputDeviceMapper::Get();
 
-	const FInputDeviceScope InputScope(nullptr, FName("GameInput"), DeviceMapper.GetUserIndexForPlatformUser(Params.PlatformUserId), GetHardwareDeviceIdentifierName(Params));
+	const FInputDeviceScope InputScope(nullptr, FName("GameInput"), DeviceMapper.GetUserIndexForPlatformUser(Params.PlatformUserId), FString(UE::GameInput::GetHardwareDeviceIdentifierName(Params.Device)));
 
 	auto ForEachValidKey = [&ButtonNames](TFunction<void(const FGamepadKeyNames::Type&)>&& Predicate)
 	{
diff --git a/Engine/Plugins/Runtime/GameInput/Source/GameInputBase/Private/Processors/IGameInputDeviceProcessor.cpp b/Engine/Plugins/Runtime/GameInput/Source/GameInputBase/Private/Processors/IGameInputDeviceProcessor.cpp
--- a/Engine/Plugins/Runtime/GameInput/Source/GameInputBase/Private/Processors/IGameInputDeviceProcessor.cpp
+++ b/Engine/Plugins/Runtime/GameInput/Source/GameInputBase/Private/Processors/IGameInputDeviceProcessor.cpp
@@ -11,11 +11,6 @@
 
 #if GAME_INPUT_SUPPORT
 
-namespace UE::GameInput
-{
-	static FName InputClassName = FName("GameInput");
-}
-
 IGameInputDeviceProcessor::IGameInputDeviceProcessor()
 {
 	InitialButtonRepeatDelay = UE::GameInput::InitialRepeatDelay;
@@ -36,53 +31,12 @@ bool IGameInputDeviceProcessor::PostProcessInput(const FGameInputEventParams& Pa
 	return false;
 }
 
-const FString& IGameInputDeviceProcessor::GetHardwareDeviceIdentifierName(const IGameInputDeviceProcessor::FGameInputEventParams& Params) const
+PRAGMA_DISABLE_DEPRECATION_WARNINGS
+const FString IGameInputDeviceProcessor::GetHardwareDeviceIdentifierName(const IGameInputDeviceProcessor::FGameInputEventParams& Params) const
 {
-	static FString ID_Virtual = TEXT("Virtual");
-	static FString ID_Aggregate = TEXT("Aggregate");
-	static FString ID_XboxOne = TEXT("XboxOne");
-	static FString ID_Xbox360 = TEXT("Xbox360");
-	static FString ID_Hid = TEXT("Hid");
-	static FString ID_I8042 = TEXT("I8042");
-
-	if (Params.Device)
-	{
-		if (const GameInputDeviceInfo* Info = Params.GetDeviceInfo())
-		{
-			// Check for any device specific overrides that may be there for custom devices
-			if (const UGameInputDeveloperSettings* Settings = GetDefault<UGameInputDeveloperSettings>())
-			{
-				if (const FGameInputDeviceConfiguration* Config = Settings->FindDeviceConfiguration(Info))
-				{
-					if (Config->bOverrideHardwareDeviceIdString)
-					{
-						return Config->OverriddenHardwareDeviceId;
-					}					
-				}
-			}
-
-			switch (Info->deviceFamily)
-			{
-			case GameInputFamilyVirtual:
-				return ID_Virtual;
-			case GameInputFamilyAggregate:
-				return ID_Aggregate;
-			case GameInputFamilyXbox360:
-				return ID_Xbox360;
-			case GameInputFamilyHid:
-				return ID_Hid;
-			case GameInputFamilyI8042:
-				return ID_I8042;
-			case GameInputFamilyXboxOne:
-			default:
-				return ID_XboxOne;
-			}
-		}
-	}
-
-	// If for some reason we are given a null device, default to just "Xbox One"
-	return ID_XboxOne;
+	return FString(UE::GameInput::GetHardwareDeviceIdentifierName(Params.Device));
 }
+PRAGMA_ENABLE_DEPRECATION_WARNINGS
 
 void IGameInputDeviceProcessor::OnControllerAnalog(const FGameInputEventParams& Params, const FName& GamePadKey, float NewAxisValueNormalized, float OldAxisValueNormalized, float DeadZone, const bool bSetDeviceScope /*= true*/)
 {
@@ -101,7 +55,12 @@ void IGameInputDeviceProcessor::OnControllerAnalog(const FGameInputEventParams&
 		{
 			if (bSetDeviceScope)
 			{
-				FInputDeviceScope InputScope(nullptr, UE::GameInput::InputClassName, IPlatformInputDeviceMapper::Get().GetUserIndexForPlatformUser(Params.PlatformUserId), GetHardwareDeviceIdentifierName(Params));
+				FInputDeviceScope InputScope(
+					nullptr,
+					UE::GameInput::GetDeviceScopeClassName(),
+					IPlatformInputDeviceMapper::Get().GetUserIndexForPlatformUser(Params.PlatformUserId),
+					FString(UE::GameInput::GetHardwareDeviceIdentifierName(Params.Device)));
+				
 				Params.MessageHandler->OnControllerAnalog(GamePadKey, Params.PlatformUserId, Params.InputDeviceId, NewAxisValueNormalized);
 			}
 			else
@@ -156,7 +115,7 @@ void IGameInputDeviceProcessor::EvaluateButtonStates(
 			// Check for button state change
 			if (0 != (ActionMask & BitMask))
 			{
-				FInputDeviceScope InputScope(nullptr, UE::GameInput::InputClassName, DeviceMapper.GetUserIndexForPlatformUser(Params.PlatformUserId), GetHardwareDeviceIdentifierName(Params));
+				FInputDeviceScope InputScope(nullptr, UE::GameInput::GetDeviceScopeClassName(), DeviceMapper.GetUserIndexForPlatformUser(Params.PlatformUserId), FString(UE::GameInput::GetHardwareDeviceIdentifierName(Params.Device)));
 
 				if (0 != (CurrentButtonHeldMask & BitMask))
 				{
@@ -175,7 +134,7 @@ void IGameInputDeviceProcessor::EvaluateButtonStates(
 			// Check for repeat key
 			if (0 != (RepeatMask & BitMask) && CurrentTime > RepeatTime[n])
 			{
-				FInputDeviceScope InputScope(nullptr, UE::GameInput::InputClassName, DeviceMapper.GetUserIndexForPlatformUser(Params.PlatformUserId), GetHardwareDeviceIdentifierName(Params));
+				FInputDeviceScope InputScope(nullptr, UE::GameInput::GetDeviceScopeClassName(), DeviceMapper.GetUserIndexForPlatformUser(Params.PlatformUserId), FString(UE::GameInput::GetHardwareDeviceIdentifierName(Params.Device)));
 
 				RepeatTime[n] = CurrentTime + ButtonRepeatDelay;
 				UE_LOG(LogGameInput, Verbose, TEXT("[FGameInputDevice::EvaluateButtonStates] (PlatformUserId = %d, InputDeviceId = %d) - Button '%s' Repeat Pressed"), Params.PlatformUserId.GetInternalId(), Params.InputDeviceId.GetId(), *ButtonKey.ToString());
diff --git a/Engine/Plugins/Runtime/GameInput/Source/GameInputBase/Public/GameInputUtils.h b/Engine/Plugins/Runtime/GameInput/Source/GameInputBase/Public/GameInputUtils.h
--- a/Engine/Plugins/Runtime/GameInput/Source/GameInputBase/Public/GameInputUtils.h
+++ b/Engine/Plugins/Runtime/GameInput/Source/GameInputBase/Public/GameInputUtils.h
@@ -66,6 +66,15 @@ namespace UE::GameInput
 	GAMEINPUTBASE_API const TArray<FGamepadKeyNames::Type>* SwitchPositionToUnrealName(const GameInputSwitchPosition ButtonMask);
 
 	GAMEINPUTBASE_API const GameInputDeviceInfo* GetDeviceInfo(IGameInputDevice* Device);
+
+	/** The FName which should be used to construct FInputDeviceScope's for game input devices. */
+	GAMEINPUTBASE_API FName GetDeviceScopeClassName();
+
+	/** 
+	* Returns the string hardware device identifier for the given input event params that can be used
+	* to create an FInputDeviceScope
+	*/
+	GAMEINPUTBASE_API FStringView GetHardwareDeviceIdentifierName(IGameInputDevice* Device);
 	
 	// Button flags assigned by GameInput
 	enum { GamingInputButtonMask = 0x3FFF };
diff --git a/Engine/Plugins/Runtime/GameInput/Source/GameInputBase/Public/Processors/IGameInputDeviceProcessor.h b/Engine/Plugins/Runtime/GameInput/Source/GameInputBase/Public/Processors/IGameInputDeviceProcessor.h
--- a/Engine/Plugins/Runtime/GameInput/Source/GameInputBase/Public/Processors/IGameInputDeviceProcessor.h
+++ b/Engine/Plugins/Runtime/GameInput/Source/GameInputBase/Public/Processors/IGameInputDeviceProcessor.h
@@ -103,7 +103,8 @@ protected:
 	* Returns the string hardware device identifier for the given input event params that can be used
 	* to create an FInputDeviceScope
 	*/
-	GAMEINPUTBASE_API const FString& GetHardwareDeviceIdentifierName(const IGameInputDeviceProcessor::FGameInputEventParams& Params) const;
+	UE_DEPRECATED(5.8, "Use UE::GameInput::GetHardwareDeviceIdentifierName instead")
+	GAMEINPUTBASE_API const FString GetHardwareDeviceIdentifierName(const IGameInputDeviceProcessor::FGameInputEventParams& Params) const;
 
 	/** A general use function to call the message handler and tell it about a controller analog key being used */
 	GAMEINPUTBASE_API void OnControllerAnalog(const FGameInputEventParams& Params, const FName& GamePadKey, float NewAxisValueNormalized, float OldAxisValueNormalized, float DeadZone, const bool bSetDeviceScope = true);
