[Game Input] Add a workaround for Game Input crashing when RenderDoc is also in use. This appears to be a bug within the GameInput library itself and is not something that UE can avoid (we can repro outside of UE in a test project), so for now we must work around it to avoid using the game input library if RenderDoc is also in use.

#Synced-CL 51759702

diff --git a/Engine/Plugins/Runtime/GameInput/Source/GameInputBase/Private/IGameInputDeviceInterface.cpp b/Engine/Plugins/Runtime/GameInput/Source/GameInputBase/Private/IGameInputDeviceInterface.cpp
--- a/Engine/Plugins/Runtime/GameInput/Source/GameInputBase/Private/IGameInputDeviceInterface.cpp
+++ b/Engine/Plugins/Runtime/GameInput/Source/GameInputBase/Private/IGameInputDeviceInterface.cpp
@@ -8,6 +8,7 @@
 #include "Misc/CoreDelegates.h"
 #include "GenericPlatform/IInputInterface.h"
 #include "HAL/IConsoleManager.h"
+#include "Misc/CommandLine.h"
 
 #if GAME_INPUT_SUPPORT
 namespace UE::GameInput
@@ -19,6 +20,42 @@ namespace UE::GameInput
 		TEXT("If true, EnumerateCurrentlyConnectedDeviceTypes will be called upon device connect and disconnect"),
 		ECVF_Default
 	);
+
+	static TAutoConsoleVariable<bool> CVarGameInputSkipInitIfRenderDocAttached
+	(
+		TEXT("gameinput.bSkipInitIfRenderDocAttached"),
+		true,
+		TEXT("If true, GameInput device callbacks will not be registered when the RenderDoc plugin is active. "
+			 "The RenderDoc plugin attempts to load the same DLLs as some of GameInput, which causes an internal abort in GameInputRedist.dll's dispatcher thread."
+			 "This issue is for GameInputRedist.dll versions 3.2+"),
+		ECVF_Default
+	);
+	
+	static bool ShouldSkipFromRenderDoc()
+	{
+		if (!UE::GameInput::CVarGameInputSkipInitIfRenderDocAttached.GetValueOnAnyThread())
+		{
+			return false;
+		}
+
+		// -attachrenderdoc on the command line triggers RenderDoc to attach at startup
+		if (FParse::Param(FCommandLine::Get(), TEXT("attachrenderdoc")))
+		{
+			return true;
+		}
+
+		// renderdoc.AutoAttach non-zero also causes RenderDoc to attach
+		if (const IConsoleVariable* AutoAttachCVar = IConsoleManager::Get().FindConsoleVariable(TEXT("renderdoc.AutoAttach")))
+		{
+			if (AutoAttachCVar->GetInt() != 0)
+			{
+				return true;
+			}
+		}
+
+		return false;
+	}
+	
 };
 #endif	// GAME_INPUT_SUPPORT
 
@@ -372,8 +409,16 @@ bool IGameInputDeviceInterface::BindDeviceStatusCallbacks()
 	
 	const GameInputKind RegisterInputKindMask = GetCurrentGameInputKindSupport();
 
-	UE_LOG(LogGameInput, Log, TEXT("Registering Device Callback for GameInputKind: '%s'. Listening for Device Status: '%s'."), *UE::GameInput::LexToString(RegisterInputKindMask), *UE::GameInput::LexToString(DeviceStatusMask));
+	UE_LOGF(LogGameInput, Log, "Registering Device Callback for GameInputKind: '%ls'. Listening for Device Status: '%ls'.", *UE::GameInput::LexToString(RegisterInputKindMask), *UE::GameInput::LexToString(DeviceStatusMask));
 	
+	// There is a crash in the internals of the GameInput library when render doc attaches,
+	// So for now, avoid attempting to bind because it will crash the editor. 
+	if (UE::GameInput::ShouldSkipFromRenderDoc())
+	{
+		UE_LOGF(LogGameInput, Warning, "RenderDoc attach detected (-attachrenderdoc on the command line, or renderdoc.AutoAttach is non-zero). Skipping GameInput device callback registration to avoid a known crash in GameInputRedist.dll (toggle 'gameinput.bSkipInitIfRenderDocAttached' to disable this and proceed anyway.).");
+		return false;
+	}
+
 	GameInput->RegisterDeviceCallback(nullptr, RegisterInputKindMask, DeviceStatusMask, GameInputBlockingEnumeration, this, DeviceCallbackFn, &ConnectionChangeCallbackToken);
 
 	// get notified when the app is constrained & unconstrained
@@ -392,10 +437,10 @@ void IGameInputDeviceInterface::OnDeviceConnectionChanged(GameInputCallbackToken
 {	
 	if (!IsSupportedDevice(Device))
 	{
-		UE_LOG(
+		UE_LOGF(
 			LogGameInput,
 			Verbose,
-			TEXT("%s is an unsupported Device and will not be processed. (Connection Changed detected from '%s' to '%s')"),
+			"%ls is an unsupported Device and will not be processed. (Connection Changed detected from '%ls' to '%ls')",
 			*UE::GameInput::LexToString(Device), *UE::GameInput::LexToString(PreviousStatus), *UE::GameInput::LexToString(CurrentStatus));
 		
 		return;
@@ -403,7 +448,7 @@ void IGameInputDeviceInterface::OnDeviceConnectionChanged(GameInputCallbackToken
 	
 	FScopeLock Lock(&DeviceInfoCS);
 	
-	UE_LOG(LogGameInput, Log, TEXT("Device Connection Changed from '%s' to '%s'"), *UE::GameInput::LexToString(PreviousStatus), *UE::GameInput::LexToString(CurrentStatus));
+	UE_LOGF(LogGameInput, Log, "Device Connection Changed from '%ls' to '%ls'", *UE::GameInput::LexToString(PreviousStatus), *UE::GameInput::LexToString(CurrentStatus));
 	
 	// This event may come in async from GameInput from outside the game thread, so defer it until Tick so we can guarantee game thread access
 	IGameInputDeviceInterface::FDeferredDeviceConnectionChanges Event = {};
@@ -460,7 +505,7 @@ void IGameInputDeviceInterface::ProcessDeferredDeviceConnectionChanges()
 			}
 			else
 			{
-				UE_LOG(LogGameInput, Error, TEXT("Unexpected state change for device %s (%s -> %s"), *UE::GameInput::LexToString(Event.Device), *UE::GameInput::LexToString(Event.PreviousStatus), *UE::GameInput::LexToString(Event.CurrentStatus));
+				UE_LOGF(LogGameInput, Error, "Unexpected state change for device %ls (%ls -> %ls", *UE::GameInput::LexToString(Event.Device), *UE::GameInput::LexToString(Event.PreviousStatus), *UE::GameInput::LexToString(Event.CurrentStatus));
 			}
 		}
 		DeferredDeviceConnectionChanges.Reset();
