[GameInput] Add delegate callback for when haptics are ready on a device. This is part 1 of a change to add audio haptic support to game input ( GAMEINPUT_HAPTIC_* types)

#rnx
#Synced-CL 52059832

diff --git a/Engine/Plugins/Runtime/GameInput/Source/GameInputBase/GameInputBase.build.cs b/Engine/Plugins/Runtime/GameInput/Source/GameInputBase/GameInputBase.build.cs
--- a/Engine/Plugins/Runtime/GameInput/Source/GameInputBase/GameInputBase.build.cs
+++ b/Engine/Plugins/Runtime/GameInput/Source/GameInputBase/GameInputBase.build.cs
@@ -17,6 +17,11 @@ namespace UnrealBuildTool.Rules
 		*/
 		protected virtual bool bSupportsDynamicLatencyInput => false;
 
+		/**
+		 * True on platforms which support Haptics through game input.
+		 */
+		protected virtual bool bSupportsHapticDevices => true;
+
 		/**
 		 * Array of required module names to be added to PublicDependencyModuleNames
 		 * if the target supports game input
@@ -51,6 +56,7 @@ namespace UnrealBuildTool.Rules
 			PublicDefinitions.Add("GAME_INPUT_SUPPORT=" + (HasGameInputSupport ? "1" : "0"));			
 
 			PublicDefinitions.Add("UE_GAME_INPUT_SUPPORTS_DLI=" + (bSupportsDynamicLatencyInput ? "1" : "0"));
+			PublicDefinitions.Add("UE_GAME_INPUT_SUPPORTS_HAPTICS=" + (bSupportsHapticDevices ? "1" : "0"));
 			
 			// Give platforms extensions a chance to add any required dependencies that may be necessary to compile game input
 			if (HasGameInputSupport)
diff --git a/Engine/Plugins/Runtime/GameInput/Source/GameInputBase/Private/GameInputDeveloperSettings.cpp b/Engine/Plugins/Runtime/GameInput/Source/GameInputBase/Private/GameInputDeveloperSettings.cpp
--- a/Engine/Plugins/Runtime/GameInput/Source/GameInputBase/Private/GameInputDeveloperSettings.cpp
+++ b/Engine/Plugins/Runtime/GameInput/Source/GameInputBase/Private/GameInputDeveloperSettings.cpp
@@ -172,10 +172,20 @@ UGameInputDeveloperSettings::UGameInputDeveloperSettings(const FObjectInitialize
 	: Super(Initializer)
 	, SupportedDeviceFamilies(static_cast<uint32>(EGameInputDeviceFamily::All))
 	, bDoNotProcessDuplicateCapabilitiesForSingleUser(true)
+	, bEnableHapticSensorSupport(false)
 {
 	PlatformSpecificSettings.Initialize(UGameInputPlatformSettings::StaticClass());
 }
 
+bool UGameInputDeveloperSettings::IsHapticSupportEnabled() const
+{
+#if UE_GAME_INPUT_SUPPORTS_HAPTICS
+	return bEnableHapticSensorSupport;
+#else
+	return false;
+#endif
+}
+
 #if WITH_EDITOR
 const TArray<uint32>& UGameInputDeveloperSettings::GetControllerButtonMappingDataKeyOptions()
 {
diff --git a/Engine/Plugins/Runtime/GameInput/Source/GameInputBase/Private/GameInputDeviceContainer.cpp b/Engine/Plugins/Runtime/GameInput/Source/GameInputBase/Private/GameInputDeviceContainer.cpp
--- a/Engine/Plugins/Runtime/GameInput/Source/GameInputBase/Private/GameInputDeviceContainer.cpp
+++ b/Engine/Plugins/Runtime/GameInput/Source/GameInputBase/Private/GameInputDeviceContainer.cpp
@@ -386,6 +386,45 @@ void FGameInputDeviceContainer::SetGameInputDevice(IGameInputDevice* InDevice)
 	}
 }
 
+IGameInputDevice* FGameInputDeviceContainer::GetHapticDevice() const
+{
+	if (!GetDefault<UGameInputDeveloperSettings>()->IsHapticSupportEnabled())
+	{
+		return nullptr;
+	}
+	
+	return HapticsDevice;
+}
+
+void FGameInputDeviceContainer::SetHapticDevice(IGameInputDevice* InDevice)
+{
+	// We should never be calling this function if haptics are not enabled
+	check(GetDefault<UGameInputDeveloperSettings>()->IsHapticSupportEnabled());
+	
+	// If the devices are the same, we can early exit. Nothing needs to happen
+	if (InDevice == HapticsDevice)
+	{
+		return;
+	}
+	
+	// TODO: Verify that this device is compatible with the actual input device somehow... maybe PID/VID?
+	HapticsDevice = InDevice;
+	
+	if (HapticsDevice)
+	{
+		// When we set a haptic device make sure we also set the device ID so that we can easily compare it later
+		if (const GameInputDeviceInfo* Info = UE::GameInput::GetDeviceInfo(HapticsDevice))
+		{
+			HapticDeviceId = Info->deviceId;
+		}
+	}
+}
+
+APP_LOCAL_DEVICE_ID FGameInputDeviceContainer::GetHapticDeviceId() const
+{
+	return HapticDeviceId;
+}
+
 APP_LOCAL_DEVICE_ID FGameInputDeviceContainer::GetGameInputDeviceId() const
 {
 	return LocalDeviceId;
diff --git a/Engine/Plugins/Runtime/GameInput/Source/GameInputBase/Private/GameInputUtils.cpp b/Engine/Plugins/Runtime/GameInput/Source/GameInputBase/Private/GameInputUtils.cpp
--- a/Engine/Plugins/Runtime/GameInput/Source/GameInputBase/Private/GameInputUtils.cpp
+++ b/Engine/Plugins/Runtime/GameInput/Source/GameInputBase/Private/GameInputUtils.cpp
@@ -171,6 +171,11 @@ namespace UE::GameInput
 		FString Result = TEXT("");
 #define DEVICE_STATUS(StatusFlag,DisplayName) if( DeviceStatus & StatusFlag ) Result += (FString(DisplayName) + TEXT("|"));
 		DEVICE_STATUS(GameInputDeviceConnected, TEXT("Connected"));
+		
+#if UE_GAME_INPUT_SUPPORTS_HAPTICS
+		DEVICE_STATUS(GameInputDeviceHapticInfoReady, TEXT("HapticInfoReady"));
+#endif
+		
 #if UE_GAMEINPUT_SUPPORTS_DEVICE_STATUS
 		DEVICE_STATUS(GameInputDeviceInputEnabled, TEXT("InputEnabled"));
 		DEVICE_STATUS(GameInputDeviceOutputEnabled, TEXT("OutputEnabled"));
diff --git a/Engine/Plugins/Runtime/GameInput/Source/GameInputBase/Private/IGameInputDeviceInterface.cpp b/Engine/Plugins/Runtime/GameInput/Source/GameInputBase/Private/IGameInputDeviceInterface.cpp
--- a/Engine/Plugins/Runtime/GameInput/Source/GameInputBase/Private/IGameInputDeviceInterface.cpp
+++ b/Engine/Plugins/Runtime/GameInput/Source/GameInputBase/Private/IGameInputDeviceInterface.cpp
@@ -405,7 +405,14 @@ bool IGameInputDeviceInterface::BindDeviceStatusCallbacks()
 	};
 
 	// TODO: should this be a member variable that subclasses can override? That might be desirable on PC
-	static constexpr GameInputDeviceStatus DeviceStatusMask = (GameInputDeviceNoStatus | GameInputDeviceConnected);
+	GameInputDeviceStatus DeviceStatusMask = (GameInputDeviceNoStatus | GameInputDeviceConnected);
+
+#if UE_GAME_INPUT_SUPPORTS_HAPTICS
+	if (GetDefault<UGameInputDeveloperSettings>()->IsHapticSupportEnabled())
+	{
+		DeviceStatusMask |= GameInputDeviceHapticInfoReady;	
+	}
+#endif
 	
 	const GameInputKind RegisterInputKindMask = GetCurrentGameInputKindSupport();
 
@@ -485,6 +492,23 @@ bool IGameInputDeviceInterface::IsSupportedDevice(IGameInputDevice* Device) cons
 	return true;
 }
 
+void IGameInputDeviceInterface::HandleHapticsReady(IGameInputDevice* Device, uint64 Timestamp)
+{
+#if UE_GAME_INPUT_SUPPORTS_HAPTICS
+	// TODO: Any shared haptic setup can be placed here
+	GameInputHapticInfo HapticInfo = {};
+	
+	HRESULT Res = Device->GetHapticInfo(&HapticInfo);
+	if (FAILED(Res))
+	{
+		UE_LOG(LogGameInput, Error, TEXT("Failed to get haptic info for device %s"), *UE::GameInput::LexToString(Device));
+		return;
+	}
+	
+	HandleHapticsReady_Impl(Device, Timestamp, HapticInfo);
+#endif	// UE_GAME_INPUT_SUPPORTS_HAPTICS
+}
+
 void IGameInputDeviceInterface::ProcessDeferredDeviceConnectionChanges()
 {
 	// We only want to actually handle device connection events in the game thread
@@ -502,6 +526,15 @@ void IGameInputDeviceInterface::ProcessDeferredDeviceConnectionChanges()
 			else if (Event.Status == EInputDeviceConnectionState::Connected)
 			{
 				HandleDeviceConnected(Event.Device, Event.Timestamp);
+
+#if UE_GAME_INPUT_SUPPORTS_HAPTICS
+				// If this event was for a connection, we also need to check if there is Haptic info ready yet.
+				// We may receive an event after the initial connection from the device to handle haptics being ready.
+				if (Event.CurrentStatus & GameInputDeviceHapticInfoReady)
+				{
+					HandleHapticsReady(Event.Device, Event.Timestamp);
+				}
+#endif
 			}
 			else
 			{
@@ -559,6 +592,9 @@ FGameInputDeviceContainer* IGameInputDeviceInterface::GetDeviceData(IGameInputDe
 		}
 	}
 
+	// TODO: Check if this device is the haptic device for one of our containers, so that we can assoiciate the 
+	// haptic device to the "normal" input device for processing
+
 	return RetVal;
 }
 
diff --git a/Engine/Plugins/Runtime/GameInput/Source/GameInputBase/Public/GameInputDeveloperSettings.h b/Engine/Plugins/Runtime/GameInput/Source/GameInputBase/Public/GameInputDeveloperSettings.h
--- a/Engine/Plugins/Runtime/GameInput/Source/GameInputBase/Public/GameInputDeveloperSettings.h
+++ b/Engine/Plugins/Runtime/GameInput/Source/GameInputBase/Public/GameInputDeveloperSettings.h
@@ -568,6 +568,24 @@ public:
 	UPROPERTY(Config, EditAnywhere, Category="Game Input")
 	bool bDoNotProcessDuplicateCapabilitiesForSingleUser;
 
+	/**
+	 * @return True if this build target supports haptics (UE_GAME_INPUT_SUPPORTS_HAPTICS) and the project setting
+	 * for enabling haptic support is true.
+	 */
+	bool IsHapticSupportEnabled() const;
+	
+protected:
+	
+	/**
+	 * If true, then Game Input will listen for the GameInputDeviceHapticInfoReady device status on platforms
+	 * which UE_GAME_INPUT_SUPPORTS_HAPTICS is true.
+	 *
+	 * Haptic sensor support is like a "fancy" force feedback. On supported devices, an audio endpoint is created.
+	 * You can use audio files to drive sound and vibration on the devices.
+	 */
+	UPROPERTY(Config, EditAnywhere, Category="Experimental|Game Input")
+	bool bEnableHapticSensorSupport;
+
 #if WITH_EDITOR
 public:
 
diff --git a/Engine/Plugins/Runtime/GameInput/Source/GameInputBase/Public/GameInputDeviceContainer.h b/Engine/Plugins/Runtime/GameInput/Source/GameInputBase/Public/GameInputDeviceContainer.h
--- a/Engine/Plugins/Runtime/GameInput/Source/GameInputBase/Public/GameInputDeviceContainer.h
+++ b/Engine/Plugins/Runtime/GameInput/Source/GameInputBase/Public/GameInputDeviceContainer.h
@@ -77,6 +77,14 @@ public:
 	*/
 	UE_API void SetGameInputDevice(IGameInputDevice* InDevice);
 
+	UE_API IGameInputDevice* GetHapticDevice() const;
+	UE_API void SetHapticDevice(IGameInputDevice* InDevice);
+
+	/**
+	 * Returns the unique local device ID of the IGameInputDevice that is used for haptics on this device container
+	 */
+	UE_API APP_LOCAL_DEVICE_ID GetHapticDeviceId() const;
+
 	/** Returns the unique local device ID of this IGameInputDevice from the GameInput API. */
 	UE_API APP_LOCAL_DEVICE_ID GetGameInputDeviceId() const;
 
@@ -106,6 +114,11 @@ protected:
 	/** The owning Game Input device that this device is associated with */
 	IGameInputDevice* Device;
 
+	/**
+	 * The game input device which is used to play haptics if UE_GAME_INPUT_SUPPORTS_HAPTICS is true.
+	 */
+	IGameInputDevice* HapticsDevice;
+
 	/** The kinds of Input that are allowed to be created for this device. */
 	GameInputKind AllowedGameInputKinds;
 
@@ -124,6 +137,8 @@ protected:
 	*/
 	APP_LOCAL_DEVICE_ID LocalDeviceId;
 
+	APP_LOCAL_DEVICE_ID HapticDeviceId;
+
 	/** Input Device processors to handle input for any specific types required (gamepad, racing wheel, KBM, etc) */
 	TArray<TSharedPtr<IGameInputDeviceProcessor>> Processors;
 
diff --git a/Engine/Plugins/Runtime/GameInput/Source/GameInputBase/Public/IGameInputDeviceInterface.h b/Engine/Plugins/Runtime/GameInput/Source/GameInputBase/Public/IGameInputDeviceInterface.h
--- a/Engine/Plugins/Runtime/GameInput/Source/GameInputBase/Public/IGameInputDeviceInterface.h
+++ b/Engine/Plugins/Runtime/GameInput/Source/GameInputBase/Public/IGameInputDeviceInterface.h
@@ -75,6 +75,25 @@ protected:
 	 */
 	virtual void HandleDeviceConnected(IGameInputDevice* Device, uint64 Timestamp) = 0;
 
+private:
+	
+	/**
+	 * Handle haptics being notified as ready from the given device.
+	 * 
+	 * This should only be called internally from this device interface in response to
+	 * a callback from the GameInput api
+	 */
+	void HandleHapticsReady(IGameInputDevice* Device, uint64 Timestamp);
+	
+protected:
+	
+	/**
+	 * Any platform specific haptic setup can be done here in an override.
+	 * 
+	 * Note: You likely want to wrap the implementation of this function in #if UE_GAME_INPUT_SUPPORTS_HAPTICS.
+	 */
+	virtual void HandleHapticsReady_Impl(IGameInputDevice* Device, uint64 Timestamp, struct GameInputHapticInfo& HapticInfo) {}
+
 	/** Process any deferred device connection changes that may have occurred. This should only be called on the game thread. */
 	GAMEINPUTBASE_API void ProcessDeferredDeviceConnectionChanges();
 	
