Accumulate accelerometer and gyro input from multiple samples per frame.

This is a simple way to improve motion through hitches and low framerates, and it also reduces interference from high-resolution haptics.

#tests PIE
#Synced-CL 52093698

#Synced-CL 52093698

diff --git a/Engine/Plugins/Runtime/GameInput/Source/GameInputBase/Private/Processors/GameInputDeviceProcessor_Sensor.cpp b/Engine/Plugins/Runtime/GameInput/Source/GameInputBase/Private/Processors/GameInputDeviceProcessor_Sensor.cpp
--- a/Engine/Plugins/Runtime/GameInput/Source/GameInputBase/Private/Processors/GameInputDeviceProcessor_Sensor.cpp
+++ b/Engine/Plugins/Runtime/GameInput/Source/GameInputBase/Private/Processors/GameInputDeviceProcessor_Sensor.cpp
@@ -13,35 +13,114 @@
 
 bool FGameInputGamepadSensorDeviceProcessor::ProcessInput(const FGameInputEventParams& Params)
 {
-	// Do nothing in this function because all sensor data is float/analog values, so we only want to
-	// check it once per-frame and not (potentially) more. 
-	return false;
-}
-
-bool FGameInputGamepadSensorDeviceProcessor::PostProcessInput(const FGameInputEventParams& Params)
-{
-	bool bRes = false;
+	// Accumulate (potentially) multiple samples per tick
 
 	// Can't do anything for an invalid platform user
 	if (!Params.PlatformUserId.IsValid() || !Params.PreviousReading)
 	{
-		return bRes;
+		return false;
 	}
 
-	// Get the current sensor state of this frame.
 	GameInputSensorsState SensorState = {};
 	if (!Params.PreviousReading->GetSensorsState(&SensorState))
 	{
-		return bRes;
+		return false;
+	}
+
+	const GameInputDeviceInfo* Info = Params.GetDeviceInfo();
+	if (!Info || !Info->sensorsInfo)
+	{
+		return false;
+	}
+
+	const GameInputSensorsKind DetectedSensorKind = Info->sensorsInfo->supportedSensors;
+
+	// This device does not support any sensors. Don't bother trying to get a reading from it.
+	if (DetectedSensorKind == GameInputSensorsNone)
+	{
+		return false;
+	}
+
+	static constexpr GameInputSensorsKind RequiredMotionDataKinds = (GameInputSensorsGyrometer | GameInputSensorsAccelerometer);
+
+	// Orientation, accel, and gyro data 
+	if (!(DetectedSensorKind & RequiredMotionDataKinds))
+	{
+		return false;
+	}
+
+	if (NumReadingsProcessedThisFrame <= 0)
+	{
+		FMemory::Memset(AccumulatedState, 0);
+		NumReadingsProcessedThisFrame = 0;
+	}
+
+	// Acceleration: we want the average of all samples from this tick
+	AccumulatedState.accelerationInGX += SensorState.accelerationInGX;
+	AccumulatedState.accelerationInGY += SensorState.accelerationInGY;
+	AccumulatedState.accelerationInGZ += SensorState.accelerationInGZ;
+
+	// Angular velocity: we want the average of all samples from this tick
+	AccumulatedState.angularVelocityInRadPerSecX += SensorState.angularVelocityInRadPerSecX;
+	AccumulatedState.angularVelocityInRadPerSecY += SensorState.angularVelocityInRadPerSecY;
+	AccumulatedState.angularVelocityInRadPerSecZ += SensorState.angularVelocityInRadPerSecZ;
+
+	// Orientation: we just want the latest
+	AccumulatedState.orientationW = SensorState.orientationW;
+	AccumulatedState.orientationX = SensorState.orientationX;
+	AccumulatedState.orientationY = SensorState.orientationY;
+	AccumulatedState.orientationZ = SensorState.orientationZ;
+
+	// Heading: we just want the latest
+	AccumulatedState.headingInDegreesFromMagneticNorth = SensorState.headingInDegreesFromMagneticNorth;
+	AccumulatedState.headingAccuracy = SensorState.headingAccuracy;
+
+	++NumReadingsProcessedThisFrame;
+
+	return true;
+}
+
+bool FGameInputGamepadSensorDeviceProcessor::PostProcessInput(const FGameInputEventParams& Params)
+{
+	if (!ProcessInput(Params) && NumReadingsProcessedThisFrame == 0)
+	{
+		return false;
 	}
 
-	// Actually process the sensor analog data 
-	bRes |= ProcessGamepadSensorState(Params, SensorState);
+	if (NumReadingsProcessedThisFrame > 1)
+	{
+		// More than one reading means we have to convert some of our accumulated inputs into averages
+		const float AverageScalar = 1.f / static_cast<float>(NumReadingsProcessedThisFrame);
+
+		// Acceleration:
+		AccumulatedState.accelerationInGX *= AverageScalar;
+		AccumulatedState.accelerationInGY *= AverageScalar;
+		AccumulatedState.accelerationInGZ *= AverageScalar;
+
+		// Angular velocity:
+		AccumulatedState.angularVelocityInRadPerSecX *= AverageScalar;
+		AccumulatedState.angularVelocityInRadPerSecY *= AverageScalar;
+		AccumulatedState.angularVelocityInRadPerSecZ *= AverageScalar;
+	}
 
-	// Track our previous sensor state for comparison on the next frame
-	PreviousState = SensorState;
+	// Negative values say "we've not yet processed any readings THIS frame, but we HAVE had readings before".
+	// In that case, the negative value tells us how many more times we should repeat the same report.
+	// Repeating the same report for a few frames is necessary when game framerate exceeds controller report rate and can also help with small hitches when the controller connection is poor.
+	if (NumReadingsProcessedThisFrame < 0)
+	{
+		++NumReadingsProcessedThisFrame;
+	}
+	else
+	{
+		static constexpr int32 ReadingRequiresRepeatReport = -30;
+		NumReadingsProcessedThisFrame = ReadingRequiresRepeatReport;
+	}
 
-	return bRes;
+	const bool bHasSensorReading = ProcessGamepadSensorState(Params, AccumulatedState);
+
+	PreviousState = AccumulatedState;
+
+	return bHasSensorReading;
 }
 
 void FGameInputGamepadSensorDeviceProcessor::ClearState(const FGameInputEventParams& Params)
@@ -53,13 +132,13 @@ void FGameInputGamepadSensorDeviceProcessor::ClearState(const FGameInputEventPar
 	}
 
 	// We can simply process a sensor state where everything is zero
-	GameInputSensorsState ZeroValueState = {};
-	FMemory::Memset(ZeroValueState, 0);
+	FMemory::Memset(AccumulatedState, 0);
+	NumReadingsProcessedThisFrame = 0;
 
-	ProcessGamepadSensorState(Params, ZeroValueState);
+	ProcessGamepadSensorState(Params, AccumulatedState);
 
 	// Make sure to update the previous state here
-	PreviousState = ZeroValueState;
+	PreviousState = AccumulatedState;
 }
 
 bool FGameInputGamepadSensorDeviceProcessor::ProcessGamepadSensorState(const FGameInputEventParams& Params, GameInputSensorsState& SensorState)
diff --git a/Engine/Plugins/Runtime/GameInput/Source/GameInputBase/Public/Processors/GameInputDeviceProcessor_Sensor.h b/Engine/Plugins/Runtime/GameInput/Source/GameInputBase/Public/Processors/GameInputDeviceProcessor_Sensor.h
--- a/Engine/Plugins/Runtime/GameInput/Source/GameInputBase/Public/Processors/GameInputDeviceProcessor_Sensor.h
+++ b/Engine/Plugins/Runtime/GameInput/Source/GameInputBase/Public/Processors/GameInputDeviceProcessor_Sensor.h
@@ -31,6 +31,18 @@ protected:
 	 * The sensor data from the previous reading that was processed
 	 */
 	GameInputSensorsState PreviousState = {};
+
+	/**
+	 * Accumulated sensor state for passing on once all input has been processed.
+	 */
+	GameInputSensorsState AccumulatedState = {};
+
+	/**
+	* Keeps track of how many times this gamepad has been processed this frame.
+	* Every successful processing of sensor input in ProcessInput will increment this value.
+	* This will get set to 0 or a negative value at the end of the input frame in PostProcessInput.
+	*/
+	int32 NumReadingsProcessedThisFrame = 0;
 };
 #endif// #if UE_GAMEINPUT_SUPPORTS_SENSORS
 
