UE_LOG -> UE_LOGF macro conversion for Game Input modules

#rnx
#Synced-CL 52309720

diff --git a/Engine/Plugins/GameInputWindows/Source/GameInputWindows/Private/GameInputWindowsDevice.cpp b/Engine/Plugins/GameInputWindows/Source/GameInputWindows/Private/GameInputWindowsDevice.cpp
--- a/Engine/Plugins/GameInputWindows/Source/GameInputWindows/Private/GameInputWindowsDevice.cpp
+++ b/Engine/Plugins/GameInputWindows/Source/GameInputWindows/Private/GameInputWindowsDevice.cpp
@@ -46,12 +46,12 @@ GameInputKind FGameInputWindowsInputDevice::GetCurrentGameInputKindSupport() con
 	{
 		if ((RegisterInputKindMask & GameInputKindKeyboard) != 0)
 		{
-			UE_LOG(LogGameInput, Log, TEXT("[FGameInputWindowsInputDevice::GetCurrentGameInputKindSupport] Keyboard support was requested, but is not currently supported via the GameInput plugin on Windows."));
+			UE_LOGF(LogGameInput, Log, "[FGameInputWindowsInputDevice::GetCurrentGameInputKindSupport] Keyboard support was requested, but is not currently supported via the GameInput plugin on Windows.");
 		}
 
 		if ((RegisterInputKindMask & GameInputKindMouse) != 0)
 		{
-			UE_LOG(LogGameInput, Log, TEXT("[FGameInputWindowsInputDevice::GetCurrentGameInputKindSupport] Mouse support was requested, but is not currently supported via the GameInput plugin on Windows."));
+			UE_LOGF(LogGameInput, Log, "[FGameInputWindowsInputDevice::GetCurrentGameInputKindSupport] Mouse support was requested, but is not currently supported via the GameInput plugin on Windows.");
 		}
 		bLogOnce = true;
 	}
@@ -75,7 +75,7 @@ void FGameInputWindowsInputDevice::HandleDeviceDisconnected(IGameInputDevice* De
 
 			// Set it's device to nullptr because it is now disconnected
 			Data->SetGameInputDevice(nullptr);
-			UE_LOG(LogGameInput, Log, TEXT("Game Input Device '%s' Disconnected Successfully at Input Device ID '%d'"), *UE::GameInput::LexToString(Device), Data->GetDeviceId().GetId());
+			UE_LOGF(LogGameInput, Log, "Game Input Device '%ls' Disconnected Successfully at Input Device ID '%d'", *UE::GameInput::LexToString(Device), Data->GetDeviceId().GetId());
 
 			// Remap this device to the "unpaired" user because it has been disconnected
 			const FPlatformUserId NewUserToAssign = IPlatformInputDeviceMapper::Get().GetUserForUnpairedInputDevices();
@@ -92,12 +92,12 @@ void FGameInputWindowsInputDevice::HandleDeviceDisconnected(IGameInputDevice* De
 		}
 		else
 		{
-			UE_LOG(LogGameInput, Error, TEXT("Game Input failed to disconnect a device. The Device '%s' did not have an associated FGameInputWindowsInputDevice!"), *UE::GameInput::LexToString(Device));
+			UE_LOGF(LogGameInput, Error, "Game Input failed to disconnect a device. The Device '%ls' did not have an associated FGameInputWindowsInputDevice!", *UE::GameInput::LexToString(Device));
 		}
 	}
 	else
 	{
-		UE_LOG(LogGameInput, Warning, TEXT("Game Input failed to disconnect a device. The Device was null! %s"), *UE::GameInput::LexToString(Device));
+		UE_LOGF(LogGameInput, Warning, "Game Input failed to disconnect a device. The Device was null! %ls", *UE::GameInput::LexToString(Device));
 	}
 
 	EnumerateCurrentlyConnectedDeviceTypes();
@@ -107,7 +107,7 @@ void FGameInputWindowsInputDevice::HandleDeviceConnected(IGameInputDevice* Devic
 {
 	// get device information
 	const GameInputDeviceInfo* Info = UE::GameInput::GetDeviceInfo(Device);
-	UE_LOG(LogGameInput, Log, TEXT("Game Input Device Connected: %s  of kind: %s"), *UE::GameInput::LexToString(Device), *UE::GameInput::LexToString(Info->supportedInput));
+	UE_LOGF(LogGameInput, Log, "Game Input Device Connected: %ls  of kind: %ls", *UE::GameInput::LexToString(Device), *UE::GameInput::LexToString(Info->supportedInput));
 	
 	FGameInputDeviceContainer* Data = GetOrCreateDeviceData(Device);
 	check(Data);
@@ -145,14 +145,14 @@ PRAGMA_ENABLE_DEPRECATION_WARNINGS
 			Data->SetPlatformUserId(UserToAssign);
 		}
 
-		UE_LOG(LogGameInput, Log, TEXT("Using PlatformUserId %d and InputDeviceId %d for device %s"),
+		UE_LOGF(LogGameInput, Log, "Using PlatformUserId %d and InputDeviceId %d for device %ls",
 			UserToAssign.GetInternalId(),
 			DeviceId.GetId(),
 			*UE::GameInput::LexToString(Device));
 	}
 	else
 	{
-		UE_LOG(LogGameInput, Log, TEXT("Game Input Device %s had no processors, so it will not be assiged to a platform user. You may need to configure it in the project settings."),
+		UE_LOGF(LogGameInput, Log, "Game Input Device %ls had no processors, so it will not be assiged to a platform user. You may need to configure it in the project settings.",
 			*UE::GameInput::LexToString(Device));
 	}
 	
@@ -222,7 +222,7 @@ void FGameInputWindowsInputDevice::HandleEditorSettingChanged()
 		IGameInputDevice* GameInputDevice = Device->GetGameInputDevice();
 		if (!IsSupportedDevice(GameInputDevice))
 		{
-			UE_LOG(LogGameInput, Log, TEXT("'%s' is not supported (the settings have changed) and input will no longer be processed for it."), *UE::GameInput::LexToString(GameInputDevice));
+			UE_LOGF(LogGameInput, Log, "'%ls' is not supported (the settings have changed) and input will no longer be processed for it.", *UE::GameInput::LexToString(GameInputDevice));
 			SupportedKind = GameInputKind::GameInputKindUnknown;
 		}
 		
diff --git a/Engine/Plugins/GameInputWindows/Source/GameInputWindows/Private/GameInputWindowsModule.cpp b/Engine/Plugins/GameInputWindows/Source/GameInputWindows/Private/GameInputWindowsModule.cpp
--- a/Engine/Plugins/GameInputWindows/Source/GameInputWindows/Private/GameInputWindowsModule.cpp
+++ b/Engine/Plugins/GameInputWindows/Source/GameInputWindows/Private/GameInputWindowsModule.cpp
@@ -52,7 +52,7 @@ TSharedPtr<IInputDevice> FGameInputWindowsModule::CreateInputDevice(const TShare
 	return DevicePtr;
 	
 #else
-	UE_LOG(LogGameInput, Error, TEXT("[FGameInputWindowsModule] Failed to create a GameInput device! GAME_INPUT_SUPPORT is false! See GameInputBase.build.cs and the build log for more info."));
+	UE_LOGF(LogGameInput, Error, "[FGameInputWindowsModule] Failed to create a GameInput device! GAME_INPUT_SUPPORT is false! See GameInputBase.build.cs and the build log for more info.");
 	return nullptr;
 #endif	// GAME_INPUT_SUPPORT
 }
diff --git a/Engine/Plugins/Runtime/GameInput/Source/GameInputBase/Private/GameInputBaseModule.cpp b/Engine/Plugins/Runtime/GameInput/Source/GameInputBase/Private/GameInputBaseModule.cpp
--- a/Engine/Plugins/Runtime/GameInput/Source/GameInputBase/Private/GameInputBaseModule.cpp
+++ b/Engine/Plugins/Runtime/GameInput/Source/GameInputBase/Private/GameInputBaseModule.cpp
@@ -40,49 +40,49 @@ bool FGameInputBaseModule::IsAvailable()
 
 void FGameInputBaseModule::StartupModule()
 {
-	UE_LOG(LogGameInput, Log, TEXT("GameInputBase module startup..."));
+	UE_LOGF(LogGameInput, Log, "GameInputBase module startup...");
 	
 	// We don't care for Game Input if we are running a commandlet, like when we are cooking.
 	if (IsRunningCommandlet())
 	{
-		UE_LOG(LogGameInput, Log, TEXT("GameInputBase module is exiting because IsRunningCommandlet is true. GameInput will not be initialized."));
+		UE_LOGF(LogGameInput, Log, "GameInputBase module is exiting because IsRunningCommandlet is true. GameInput will not be initialized.");
 		return;
 	}
 
 	// If there is no project name then we don't need game input either. This means we are in the project launcher
 	if (!FApp::HasProjectName())
 	{
-		UE_LOG(LogGameInput, Log, TEXT("GameInputBase module is exiting because there is no project name. GameInput will not be initialized."));
+		UE_LOGF(LogGameInput, Log, "GameInputBase module is exiting because there is no project name. GameInput will not be initialized.");
 		return;
 	}
 
 	// Unattended app can't receive any user input, so there is no need to try and create the GameInput interface.
 	if (FApp::IsUnattended() && !FApp::AllowUnattendedInput())
 	{
-		UE_LOG(LogGameInput, Log, TEXT("GameInputBase module is exiting because it is unattended (FApp::IsUnattended is true) and thus cannot receive user input. GameInput will not be initialized."));
+		UE_LOGF(LogGameInput, Log, "GameInputBase module is exiting because it is unattended (FApp::IsUnattended is true) and thus cannot receive user input. GameInput will not be initialized.");
 		return;
 	}
 
 	// Doesn't make sense to have headless apps create game input
 	if (!FApp::CanEverRender())
 	{
-		UE_LOG(LogGameInput, Log, TEXT("GameInputBase module is exiting because it cannot render anything (FApp::CanEverRender is false). GameInput will not be initialized."));
+		UE_LOGF(LogGameInput, Log, "GameInputBase module is exiting because it cannot render anything (FApp::CanEverRender is false). GameInput will not be initialized.");
 		return;
 	}
 
 #if GAME_INPUT_SUPPORT
 
-	UE_LOG(LogGameInput, Log, TEXT("Compiled using GameInput major Version: %d"), GAMEINPUT_API_VERSION);
+	UE_LOGF(LogGameInput, Log, "Compiled using GameInput major Version: %d", GAMEINPUT_API_VERSION);
 	
 #if PLATFORM_WINDOWS
 
 	if (!UE::GameInput::HasValidVersionOfGameInput())
 	{
-		UE_LOG(LogGameInput, Warning, TEXT("[%hs] module is exiting because a valid GameInput version is not installed. GameInput will not be initialized. Is it installed correctly?"), __func__);
+		UE_LOGF(LogGameInput, Warning, "[%hs] module is exiting because a valid GameInput version is not installed. GameInput will not be initialized. Is it installed correctly?", __func__);
 		return;
 	}
 
-	UE_LOG(LogGameInput, Log, TEXT("Executing with detected GameInput Version: %s"), *UE::GameInput::GetInstalledGameInputVersionInfoString());
+	UE_LOGF(LogGameInput, Log, "Executing with detected GameInput Version: %ls", *UE::GameInput::GetInstalledGameInputVersionInfoString());
 
 // Disable warning C4191: 'type cast' : unsafe conversion from 'PROC' to 'XXX' getting the GameInputCreate function
 #pragma warning(push)
@@ -99,11 +99,11 @@ void FGameInputBaseModule::StartupModule()
 		
 		if (SUCCEEDED(HResult))
 		{
-			UE_LOG(LogGameInput, Log, TEXT("[FGameInputBaseModule::StartupModule] Successfully created the IGameInput interface"));
+			UE_LOGF(LogGameInput, Log, "[FGameInputBaseModule::StartupModule] Successfully created the IGameInput interface");
 		}
 		else
 		{
-			UE_LOG(LogGameInput, Warning, TEXT("Failed to initialize GameInput: 0x%X"), HResult);	
+			UE_LOGF(LogGameInput, Warning, "Failed to initialize GameInput: 0x%X", HResult);	
 		}
 
 		struct FBroadcastCreationDelegateFunctor
@@ -144,11 +144,11 @@ void FGameInputBaseModule::StartupModule()
 	const HRESULT HResult = GameInputCreate(&UE::GameInput::GGameInputInterface);
 	if (SUCCEEDED(HResult))
 	{
-		UE_LOG(LogGameInput, Log, TEXT("[FGameInputBaseModule::StartupModule] Successfully created the IGameInput interface"));
+		UE_LOGF(LogGameInput, Log, "[FGameInputBaseModule::StartupModule] Successfully created the IGameInput interface");
 	}
 	else
 	{
-		UE_LOG(LogGameInput, Error, TEXT("Failed to initialize GameInput: 0x%X"), HResult);
+		UE_LOGF(LogGameInput, Error, "Failed to initialize GameInput: 0x%X", HResult);
 	}
 	OnGameInputCreation.Broadcast(UE::GameInput::GGameInputInterface);
 	
@@ -157,7 +157,7 @@ void FGameInputBaseModule::StartupModule()
 	InitializeGameInputKeys();
 	
 #else
-	UE_LOG(LogGameInput, Warning, TEXT("Failed to initialize GameInput! GAME_INPUT_SUPPORT is false!"));
+	UE_LOGF(LogGameInput, Warning, "Failed to initialize GameInput! GAME_INPUT_SUPPORT is false!");
 #endif	// #if GAME_INPUT_SUPPORT
 }
 
diff --git a/Engine/Plugins/Runtime/GameInput/Source/GameInputBase/Private/GameInputDeviceContainer.cpp b/Engine/Plugins/Runtime/GameInput/Source/GameInputBase/Private/GameInputDeviceContainer.cpp
--- a/Engine/Plugins/Runtime/GameInput/Source/GameInputBase/Private/GameInputDeviceContainer.cpp
+++ b/Engine/Plugins/Runtime/GameInput/Source/GameInputBase/Private/GameInputDeviceContainer.cpp
@@ -63,7 +63,7 @@ void FGameInputDeviceContainer::InitalizeDeviceProcessors()
 
 	const GameInputKind DeviceKind = Info->supportedInput;
 
-	UE_LOG(LogGameInput, Log, TEXT("InitalizeDeviceProcessors for device kind %s with an allowed kind of %s"), *UE::GameInput::LexToString(DeviceKind), *UE::GameInput::LexToString(AllowedGameInputKinds));
+	UE_LOGF(LogGameInput, Log, "InitalizeDeviceProcessors for device kind %ls with an allowed kind of %ls", *UE::GameInput::LexToString(DeviceKind), *UE::GameInput::LexToString(AllowedGameInputKinds));
 
 	InitalizeDeviceProcessors_Impl();
 }
@@ -138,10 +138,10 @@ void FGameInputDeviceContainer::InitalizeDeviceProcessors_Impl()
 		}
 		else
 		{
-			UE_LOG(
+			UE_LOGF(
 				LogGameInput,
 				Warning,
-				TEXT("A game input controller device (%s) was connected but will not be processed because it does not have an explict device configuration"),
+				"A game input controller device (%ls) was connected but will not be processed because it does not have an explict device configuration",
 				*UE::GameInput::LexToString(Device));
 		}
 	}
@@ -161,10 +161,10 @@ void FGameInputDeviceContainer::InitalizeDeviceProcessors_Impl()
 		}
 		else
 		{
-			UE_LOG(
+			UE_LOGF(
 				LogGameInput,
 				Warning,
-				TEXT("A raw game input device (%s) was connected but will not be processed because it does not have an explicit device configuration"),
+				"A raw game input device (%ls) was connected but will not be processed because it does not have an explicit device configuration",
 				*UE::GameInput::LexToString(Device));
 		}
 	}
diff --git a/Engine/Plugins/Runtime/GameInput/Source/GameInputBase/Private/IGameInputDeviceInterface.cpp b/Engine/Plugins/Runtime/GameInput/Source/GameInputBase/Private/IGameInputDeviceInterface.cpp
--- a/Engine/Plugins/Runtime/GameInput/Source/GameInputBase/Private/IGameInputDeviceInterface.cpp
+++ b/Engine/Plugins/Runtime/GameInput/Source/GameInputBase/Private/IGameInputDeviceInterface.cpp
@@ -501,7 +501,7 @@ void IGameInputDeviceInterface::HandleHapticsReady(IGameInputDevice* Device, uin
 	HRESULT Res = Device->GetHapticInfo(&HapticInfo);
 	if (FAILED(Res))
 	{
-		UE_LOG(LogGameInput, Error, TEXT("Failed to get haptic info for device %s"), *UE::GameInput::LexToString(Device));
+		UE_LOGF(LogGameInput, Error, "Failed to get haptic info for device %ls", *UE::GameInput::LexToString(Device));
 		return;
 	}
 	
diff --git a/Engine/Plugins/Runtime/GameInput/Source/GameInputBase/Private/Processors/GameInputDeviceProcessor_Controller.cpp b/Engine/Plugins/Runtime/GameInput/Source/GameInputBase/Private/Processors/GameInputDeviceProcessor_Controller.cpp
--- a/Engine/Plugins/Runtime/GameInput/Source/GameInputBase/Private/Processors/GameInputDeviceProcessor_Controller.cpp
+++ b/Engine/Plugins/Runtime/GameInput/Source/GameInputBase/Private/Processors/GameInputDeviceProcessor_Controller.cpp
@@ -276,7 +276,7 @@ bool FGameInputControllerDeviceProcessor::ProcessControllerButtonState(const FGa
 
 		if (ButtonStates[i])
 		{
-			UE_LOG(LogGameInput, Verbose, TEXT("[ProcessControllerButtonState] Device ID: %d   Button count: %d    index: %d  State: %d   %s"), 
+			UE_LOGF(LogGameInput, Verbose, "[ProcessControllerButtonState] Device ID: %d   Button count: %d    index: %d  State: %d   %ls", 
 				Params.InputDeviceId.GetId(),
 				ButtonCount, 
 				i,
@@ -351,20 +351,20 @@ bool FGameInputControllerDeviceProcessor::ProcessControllerAxisState(const FGame
 				}
 				else
 				{
-					UE_LOG(LogGameInput, VeryVerbose, TEXT("[ProcessControllerAxisState] (Device %s) Invalid key name configured for controller axis '%d': %.3f"), *UE::GameInput::LexToString(Params.Device), i, CurrentValue);
+					UE_LOGF(LogGameInput, VeryVerbose, "[ProcessControllerAxisState] (Device %ls) Invalid key name configured for controller axis '%d': %.3f", *UE::GameInput::LexToString(Params.Device), i, CurrentValue);
 				}
 			}
 			else
 			{
 				// TODO: Here is where could send a "Generic USB Axis X" key here which could allow for us to support many more devices via a key rebind screen
-				UE_LOG(LogGameInput, VeryVerbose, TEXT("[ProcessControllerAxisState] (Device %s) Controller axis '%d' has value: %.3f"), *UE::GameInput::LexToString(Params.Device), i, CurrentValue);
+				UE_LOGF(LogGameInput, VeryVerbose, "[ProcessControllerAxisState] (Device %ls) Controller axis '%d' has value: %.3f", *UE::GameInput::LexToString(Params.Device), i, CurrentValue);
 			}
 		}		
 		// You are receiving analog values from an axis that you might not know about, log it here
 		// in case you are trying to set something up
 		else
 		{
-			UE_LOG(LogGameInput, VeryVerbose, TEXT("[ProcessControllerAxisState] (Device %s) Receiving input from an unconfigured controller axis '%d': %.3f"), *UE::GameInput::LexToString(Params.Device), i, CurrentValue);
+			UE_LOGF(LogGameInput, VeryVerbose, "[ProcessControllerAxisState] (Device %ls) Receiving input from an unconfigured controller axis '%d': %.3f", *UE::GameInput::LexToString(Params.Device), i, CurrentValue);
 		}
 	}
 
diff --git a/Engine/Plugins/Runtime/GameInput/Source/GameInputBase/Private/Processors/GameInputDeviceProcessor_KeyboardMouse.cpp b/Engine/Plugins/Runtime/GameInput/Source/GameInputBase/Private/Processors/GameInputDeviceProcessor_KeyboardMouse.cpp
--- a/Engine/Plugins/Runtime/GameInput/Source/GameInputBase/Private/Processors/GameInputDeviceProcessor_KeyboardMouse.cpp
+++ b/Engine/Plugins/Runtime/GameInput/Source/GameInputBase/Private/Processors/GameInputDeviceProcessor_KeyboardMouse.cpp
@@ -58,19 +58,19 @@ void FGameInputKeyboardDeviceProcessor::UpdateUnifiedKeyboardState(const FGameIn
 		{
 			KeyRepeatTime.Add(KeyCode, CurrentTime + GKeyboardRepeatInitialDelay);
 			Params.MessageHandler->OnKeyDown(KeyCode, 0, false);
-			UE_LOG(LogGameInput, Verbose, TEXT("Key Press 0x%X"), KeyCode);
+			UE_LOGF(LogGameInput, Verbose, "Key Press 0x%X", KeyCode);
 
 			if (KeyCode == VK_CAPITAL)
 			{
 				SetSimulatedCapsLock(!bSimulatedCapsLock);
-				UE_LOG(LogGameInput, Verbose, TEXT("Simulated caps lock is %s"), bSimulatedCapsLock ? TEXT("ON") : TEXT("OFF"));
+				UE_LOGF(LogGameInput, Verbose, "Simulated caps lock is %ls", bSimulatedCapsLock ? TEXT("ON") : TEXT("OFF"));
 			}
 		}
 		else if (CurrentTime > KeyRepeatTime[KeyCode])
 		{
 			KeyRepeatTime.Add(KeyCode, CurrentTime + GKeyboardRepeatDelay);
 			Params.MessageHandler->OnKeyDown(KeyCode, 0, true);
-			UE_LOG(LogGameInput, Verbose, TEXT("Key Press 0x%X (repeat)"), KeyCode);
+			UE_LOGF(LogGameInput, Verbose, "Key Press 0x%X (repeat)", KeyCode);
 		}
 	}
 
@@ -81,7 +81,7 @@ void FGameInputKeyboardDeviceProcessor::UpdateUnifiedKeyboardState(const FGameIn
 		{
 			KeyRepeatTime.Remove(KeyCode);
 			Params.MessageHandler->OnKeyUp(KeyCode, 0, false);
-			UE_LOG(LogGameInput, Verbose, TEXT("Key Release 0x%X"), KeyCode);
+			UE_LOGF(LogGameInput, Verbose, "Key Release 0x%X", KeyCode);
 		}
 	}
 
@@ -230,7 +230,7 @@ bool FGameInputMouseDeviceProcessor::ProcessInput(const FGameInputEventParams& P
 		LastMouseOffset.X += MouseDX;
 		LastMouseOffset.Y += MouseDY;
 
-		UE_LOG(LogGameInput, VeryVerbose, TEXT("Device %s (InputDeviceId = %d) - Mouse RawDX %d, RawDY %d, DX %0.2f, DY %.2f  %s"), *UE::GameInput::LexToString(Params.Device), Params.InputDeviceId.GetId(), RawMouseDX, RawMouseDY, MouseDX, MouseDY, bHighPrecisionMouseMode ? TEXT("(high precision)") : TEXT(""));
+		UE_LOGF(LogGameInput, VeryVerbose, "Device %ls (InputDeviceId = %d) - Mouse RawDX %d, RawDY %d, DX %0.2f, DY %.2f  %ls", *UE::GameInput::LexToString(Params.Device), Params.InputDeviceId.GetId(), RawMouseDX, RawMouseDY, MouseDX, MouseDY, bHighPrecisionMouseMode ? TEXT("(high precision)") : TEXT(""));
 
 		// OnRawMouseMove may be used to process player character controls so, as the name suggests, these should be raw data without any smoothing.
 		Params.MessageHandler->OnRawMouseMove(RawMouseDX, RawMouseDY);
@@ -242,7 +242,7 @@ bool FGameInputMouseDeviceProcessor::ProcessInput(const FGameInputEventParams& P
 	float MouseWheelDY = (float)(MouseState.wheelY - PreviousMouseState.wheelY);
 	if (MouseWheelDY != 0)
 	{
-		UE_LOG(LogGameInput, Verbose, TEXT("Device %s (InputDeviceId = %d) - Mouse Wheel DY %.2f"), *UE::GameInput::LexToString(Params.Device), Params.InputDeviceId.GetId(), MouseWheelDY);
+		UE_LOGF(LogGameInput, Verbose, "Device %ls (InputDeviceId = %d) - Mouse Wheel DY %.2f", *UE::GameInput::LexToString(Params.Device), Params.InputDeviceId.GetId(), MouseWheelDY);
 		Params.MessageHandler->OnMouseWheel(MouseWheelDY * MouseWheelSpinFactor);
 	}
 
@@ -266,12 +266,12 @@ bool FGameInputMouseDeviceProcessor::ProcessInput(const FGameInputEventParams& P
 			const bool bHasDoubleClick = (GMouseDoubleClickDelay > 0) && (CurrentTime <= RepeatTime[ButtonIndex]) && (LastMouseOffset.SizeSquared() <= FMath::Square(GMouseDoubleClickArea));
 			if (bHasDoubleClick)
 			{
-				UE_LOG(LogGameInput, Verbose, TEXT("Device %s (DeviceIndex = %d) - %s Double Click"), *UE::GameInput::LexToString(Params.Device), Params.InputDeviceId.GetId(), UE::GameInput::GetMouseButtonName(MouseButton));
+				UE_LOGF(LogGameInput, Verbose, "Device %ls (DeviceIndex = %d) - %ls Double Click", *UE::GameInput::LexToString(Params.Device), Params.InputDeviceId.GetId(), UE::GameInput::GetMouseButtonName(MouseButton));
 				Params.MessageHandler->OnMouseDoubleClick(nullptr, MouseButton);
 			}
 			else
 			{
-				UE_LOG(LogGameInput, Verbose, TEXT("Device %s (DeviceIndex = %d) - %s Pressed"), *UE::GameInput::LexToString(Params.Device), Params.InputDeviceId.GetId(), UE::GameInput::GetMouseButtonName(MouseButton));
+				UE_LOGF(LogGameInput, Verbose, "Device %ls (DeviceIndex = %d) - %ls Pressed", *UE::GameInput::LexToString(Params.Device), Params.InputDeviceId.GetId(), UE::GameInput::GetMouseButtonName(MouseButton));
 				Params.MessageHandler->OnMouseDown(nullptr, MouseButton);
 			}
 
@@ -280,7 +280,7 @@ bool FGameInputMouseDeviceProcessor::ProcessInput(const FGameInputEventParams& P
 		}
 		else
 		{
-			UE_LOG(LogGameInput, Verbose, TEXT("Device %s (DeviceIndex = %d) - %s Released"), *UE::GameInput::LexToString(Params.Device), Params.InputDeviceId.GetId(), UE::GameInput::GetMouseButtonName(MouseButton));
+			UE_LOGF(LogGameInput, Verbose, "Device %ls (DeviceIndex = %d) - %ls Released", *UE::GameInput::LexToString(Params.Device), Params.InputDeviceId.GetId(), UE::GameInput::GetMouseButtonName(MouseButton));
 			Params.MessageHandler->OnMouseUp(MouseButton);
 		}
 	}
@@ -300,7 +300,7 @@ void FGameInputMouseDeviceProcessor::ClearState(const FGameInputEventParams& Par
 		if ((GameInputButton & LastButtonHeldMask) != 0)
 		{
 			EMouseButtons::Type MouseButton = UE::GameInput::MouseButtonMappings[ButtonIndex].MouseButton;
-			UE_LOG(LogGameInput, Verbose, TEXT("Device %s (InputDeviceId = %d) - %s Released (via ClearState)"), *UE::GameInput::LexToString(Params.Device), Params.InputDeviceId.GetId(), UE::GameInput::GetMouseButtonName(MouseButton));
+			UE_LOGF(LogGameInput, Verbose, "Device %ls (InputDeviceId = %d) - %ls Released (via ClearState)", *UE::GameInput::LexToString(Params.Device), Params.InputDeviceId.GetId(), UE::GameInput::GetMouseButtonName(MouseButton));
 			Params.MessageHandler->OnMouseUp(MouseButton);
 		}
 	}
diff --git a/Engine/Plugins/Runtime/GameInput/Source/GameInputBase/Private/Processors/GameInputDeviceProcessor_Raw.cpp b/Engine/Plugins/Runtime/GameInput/Source/GameInputBase/Private/Processors/GameInputDeviceProcessor_Raw.cpp
--- a/Engine/Plugins/Runtime/GameInput/Source/GameInputBase/Private/Processors/GameInputDeviceProcessor_Raw.cpp
+++ b/Engine/Plugins/Runtime/GameInput/Source/GameInputBase/Private/Processors/GameInputDeviceProcessor_Raw.cpp
@@ -59,7 +59,7 @@ const GameInputRawDeviceReportInfo* FGameInputRawDeviceProcessor::Internal_Attem
 	
 	if (!ReadingToUse)
 	{
-		UE_LOG(LogGameInput, Error, TEXT("[ReadCurrentRawInputState] Cannot read raw input state, ReadingToUse was null (Device %s) "), *UE::GameInput::LexToString(Params.Device));
+		UE_LOGF(LogGameInput, Error, "[ReadCurrentRawInputState] Cannot read raw input state, ReadingToUse was null (Device %ls) ", *UE::GameInput::LexToString(Params.Device));
 		return nullptr;
 	}
 
@@ -67,7 +67,7 @@ const GameInputRawDeviceReportInfo* FGameInputRawDeviceProcessor::Internal_Attem
 	const bool bSuccessfulReading = ReadingToUse->GetRawReport(&RawReport);
 	if (!bSuccessfulReading)
 	{
-		UE_LOG(LogGameInput, Error, TEXT("[ReadCurrentRawInputState] Unsuccessful reading of raw input report! (GetRawReport failed) (Device %s) "), *UE::GameInput::LexToString(Params.Device));
+		UE_LOGF(LogGameInput, Error, "[ReadCurrentRawInputState] Unsuccessful reading of raw input report! (GetRawReport failed) (Device %ls) ", *UE::GameInput::LexToString(Params.Device));
 		return nullptr;
 	}
 
@@ -89,7 +89,7 @@ const GameInputRawDeviceReportInfo* FGameInputRawDeviceProcessor::Internal_Attem
 	const GameInputRawDeviceReportInfo* RawReportInfo = RawReport->GetReportInfo();
 	if (!RawReportInfo)
 	{
-		UE_LOG(LogGameInput, Warning, TEXT("[ProcessRawReport] Unsuccessful reading of raw input report! (GameInputRawDeviceReportInfo is null) (Device %s) "), *UE::GameInput::LexToString(Params.Device));
+		UE_LOGF(LogGameInput, Warning, "[ProcessRawReport] Unsuccessful reading of raw input report! (GameInputRawDeviceReportInfo is null) (Device %ls) ", *UE::GameInput::LexToString(Params.Device));
 	}
 
 	return RawReportInfo;
@@ -107,7 +107,7 @@ GameInputRawDeviceReportInfo FGameInputRawDeviceProcessor::ReadCurrentRawInputSt
 #if GAMEINPUT_API_VERSION >= 3
 	if (!ReadingToUse)
 	{
-		UE_LOG(LogGameInput, Error, TEXT("[ReadCurrentRawInputState] Cannot read raw input state, ReadingToUse was null (Device %s) "), *UE::GameInput::LexToString(Params.Device));
+		UE_LOGF(LogGameInput, Error, "[ReadCurrentRawInputState] Cannot read raw input state, ReadingToUse was null (Device %ls) ", *UE::GameInput::LexToString(Params.Device));
 		return RawReportInfo;
 	}
 
@@ -115,7 +115,7 @@ GameInputRawDeviceReportInfo FGameInputRawDeviceProcessor::ReadCurrentRawInputSt
 	const bool bSuccessfulReading = ReadingToUse->GetRawReport(&RawReport);
 	if (!bSuccessfulReading)
 	{
-		UE_LOG(LogGameInput, Error, TEXT("[ReadCurrentRawInputState] Unsuccessful reading of raw input report! (GetRawReport failed) (Device %s) "), *UE::GameInput::LexToString(Params.Device));
+		UE_LOGF(LogGameInput, Error, "[ReadCurrentRawInputState] Unsuccessful reading of raw input report! (GetRawReport failed) (Device %ls) ", *UE::GameInput::LexToString(Params.Device));
 		return RawReportInfo;
 	}
 
@@ -146,7 +146,7 @@ GameInputRawDeviceReportInfo FGameInputRawDeviceProcessor::ReadCurrentRawInputSt
 	}
 
 #else
-	UE_LOG(LogGameInput, Error, TEXT("A GameInput version of 3+ is required for raw device support. (Device %s) "), *UE::GameInput::LexToString(Params.Device));
+	UE_LOGF(LogGameInput, Error, "A GameInput version of 3+ is required for raw device support. (Device %ls) ", *UE::GameInput::LexToString(Params.Device));
 #endif	// #if GAMEINPUT_API_VERSION >= 3
 	
 	return RawReportInfo;
@@ -167,7 +167,7 @@ bool FGameInputRawDeviceProcessor::ProcessInput(const FGameInputEventParams& Par
 	// Check that we have a valid config before bothering to read the raw report
 	if (!DeviceConfig)
 	{
-		UE_LOG(LogGameInput, Verbose, TEXT("[ProcessRawReport] (Device %s) Does not have a valid FGameInputDeviceConfiguration in the UGameInputDeveloperSettings. We can't process Raw Input without it. Exiting."), *UE::GameInput::LexToString(Params.Device));
+		UE_LOGF(LogGameInput, Verbose, "[ProcessRawReport] (Device %ls) Does not have a valid FGameInputDeviceConfiguration in the UGameInputDeveloperSettings. We can't process Raw Input without it. Exiting.", *UE::GameInput::LexToString(Params.Device));
 		return bRes;
 	}
 
@@ -215,7 +215,7 @@ bool FGameInputRawDeviceProcessor::PostProcessInput(const FGameInputEventParams&
 	// Check that we have a valid config before bothering to read the raw report
 	if (!DeviceConfig)
 	{
-		UE_LOG(LogGameInput, Verbose, TEXT("[ProcessRawReport] (Device %s) Does not have a valid FGameInputDeviceConfiguration in the UGameInputDeveloperSettings. We can't process Raw Input without it. Exiting."), *UE::GameInput::LexToString(Params.Device));
+		UE_LOGF(LogGameInput, Verbose, "[ProcessRawReport] (Device %ls) Does not have a valid FGameInputDeviceConfiguration in the UGameInputDeveloperSettings. We can't process Raw Input without it. Exiting.", *UE::GameInput::LexToString(Params.Device));
 		return bRes;
 	}
 
@@ -303,12 +303,12 @@ bool FGameInputRawDeviceProcessor::ProcessAllRawValues(const FGameInputEventPara
 			else
 			{
 				// You want a valid key name here, throw a warning if your config is wrong
-				UE_LOG(LogGameInput, Warning, TEXT("[ProcessRawReport] Invalid key name for raw report axis at index %d with value of %u (Device %s)"), i, Val, *UE::GameInput::LexToString(Params.Device));
+				UE_LOGF(LogGameInput, Warning, "[ProcessRawReport] Invalid key name for raw report axis at index %d with value of %u (Device %ls)", i, Val, *UE::GameInput::LexToString(Params.Device));
 			}
 		}
 		else if (Val > 0)
 		{
-			UE_LOG(LogGameInput, VeryVerbose, TEXT("[ProcessRawReport] (Device %s) No raw device report config for axis '%d' with value of %u"), *UE::GameInput::LexToString(Params.Device), i, Val);
+			UE_LOGF(LogGameInput, VeryVerbose, "[ProcessRawReport] (Device %ls) No raw device report config for axis '%d' with value of %u", *UE::GameInput::LexToString(Params.Device), i, Val);
 		}
 	}
 
@@ -395,7 +395,7 @@ bool FGameInputRawDeviceProcessor::ProcessRawInputValueAsAanalog(const FGameInpu
 	// A valid FKey is required
 	if (!FKey(AxisData->KeyName).IsValid())
 	{
-		UE_LOG(LogGameInput, Error, TEXT("Attempting to report an invalid key name '%s' (Device %s) "), *AxisData->KeyName.ToString(), *UE::GameInput::LexToString(Params.Device));
+		UE_LOGF(LogGameInput, Error, "Attempting to report an invalid key name '%ls' (Device %ls) ", *AxisData->KeyName.ToString(), *UE::GameInput::LexToString(Params.Device));
 		return false;
 	}
 
@@ -494,7 +494,7 @@ void FGameInputRawDeviceProcessor::EvaluatePackedUInt8ButtonChordsAsAnalog(
 	// We can't tell slate about input messages from an invalid platform user
 	if (!Params.PlatformUserId.IsValid())
 	{
-		UE_LOG(LogGameInput, Verbose, TEXT("Attempting to evaluate button states with an invalid platform user id of '%d'. The button messages will not be sent."), Params.PlatformUserId.GetInternalId());
+		UE_LOGF(LogGameInput, Verbose, "Attempting to evaluate button states with an invalid platform user id of '%d'. The button messages will not be sent.", Params.PlatformUserId.GetInternalId());
 		return;
 	}
 
@@ -520,7 +520,7 @@ PRAGMA_ENABLE_DEPRECATION_WARNINGS
 				continue;
 			}
 			
-			UE_LOG(LogGameInput, Verbose, TEXT("[FGameInputRawDeviceProcessor::EvaluatePackedUInt8ButtonChordsAsAnalog] (PlatformUserId = %d, InputDeviceId = %d) - Button '%s' Value %.3f"),
+			UE_LOGF(LogGameInput, Verbose, "[FGameInputRawDeviceProcessor::EvaluatePackedUInt8ButtonChordsAsAnalog] (PlatformUserId = %d, InputDeviceId = %d) - Button '%ls' Value %.3f",
 				Params.PlatformUserId.GetInternalId(), Params.InputDeviceId.GetId(), *KeyName.ToString(), CurrentValueFloat);
 			
 			Params.MessageHandler->OnControllerAnalog(KeyName, Params.PlatformUserId, Params.InputDeviceId, CurrentValueFloat);	
@@ -538,13 +538,13 @@ void FGameInputRawDeviceProcessor::EvaluatePackedUInt8ButtonChordsAsButtonEvents
 	// We can't tell slate about input messages from an invalid platform user
 	if (!Params.PlatformUserId.IsValid())
 	{
-		UE_LOG(LogGameInput, Verbose, TEXT("Attempting to evaluate button states with an invalid platform user id of '%d'. The button messages will not be sent."), Params.PlatformUserId.GetInternalId());
+		UE_LOGF(LogGameInput, Verbose, "Attempting to evaluate button states with an invalid platform user id of '%d'. The button messages will not be sent.", Params.PlatformUserId.GetInternalId());
 		return;
 	}
 
 	if (!RepeatTime)
 	{
-		UE_LOG(LogGameInput, Error, TEXT("[IGameInputDeviceProcessor::EvaluateButtonStates] Invalid RepeatTime array given to evaluate button states!"));
+		UE_LOGF(LogGameInput, Error, "[IGameInputDeviceProcessor::EvaluateButtonStates] Invalid RepeatTime array given to evaluate button states!");
 		return;
 	}
 
@@ -581,7 +581,7 @@ PRAGMA_ENABLE_DEPRECATION_WARNINGS
 		{
 			ForEachValidKey([&Params](const FGamepadKeyNames::Type& KeyName)
 			{
-				UE_LOG(LogGameInput, Verbose, TEXT("[FGameInputDevice::EvaluateButtonStates] (PlatformUserId = %d, InputDeviceId = %d) - Button '%s' Pressed"),
+				UE_LOGF(LogGameInput, Verbose, "[FGameInputDevice::EvaluateButtonStates] (PlatformUserId = %d, InputDeviceId = %d) - Button '%ls' Pressed",
 					Params.PlatformUserId.GetInternalId(), Params.InputDeviceId.GetId(), *KeyName.ToString());
 					
 				Params.MessageHandler->OnControllerButtonPressed(KeyName, Params.PlatformUserId, Params.InputDeviceId, false);
@@ -593,7 +593,7 @@ PRAGMA_ENABLE_DEPRECATION_WARNINGS
 		{
 			ForEachValidKey([&Params](const FGamepadKeyNames::Type& KeyName)
 			{
-				UE_LOG(LogGameInput, Verbose, TEXT("[FGameInputDevice::EvaluateButtonStates] (PlatformUserId = %d, InputDeviceId = %d) - Button '%s' Released"),
+				UE_LOGF(LogGameInput, Verbose, "[FGameInputDevice::EvaluateButtonStates] (PlatformUserId = %d, InputDeviceId = %d) - Button '%ls' Released",
 					Params.PlatformUserId.GetInternalId(), Params.InputDeviceId.GetId(), *KeyName.ToString());
 					
 				Params.MessageHandler->OnControllerButtonReleased(KeyName, Params.PlatformUserId, Params.InputDeviceId, false);
@@ -610,7 +610,7 @@ PRAGMA_ENABLE_DEPRECATION_WARNINGS
 		
 		ForEachValidKey([&Params](const FGamepadKeyNames::Type& KeyName)
 		{
-			UE_LOG(LogGameInput, Verbose, TEXT("[FGameInputDevice::EvaluateButtonStates] (PlatformUserId = %d, InputDeviceId = %d) - Button '%s' Repeat Pressed"),
+			UE_LOGF(LogGameInput, Verbose, "[FGameInputDevice::EvaluateButtonStates] (PlatformUserId = %d, InputDeviceId = %d) - Button '%ls' Repeat Pressed",
 				Params.PlatformUserId.GetInternalId(), Params.InputDeviceId.GetId(), *KeyName.ToString());
 				
 			Params.MessageHandler->OnControllerButtonPressed(KeyName, Params.PlatformUserId, Params.InputDeviceId, true);
@@ -631,7 +631,7 @@ void FGameInputRawDeviceProcessor::ClearState(const FGameInputEventParams& Param
 	// Check that we have a valid config before bothering to read the raw report
 	if (!DeviceConfig)
 	{
-		UE_LOG(LogGameInput, Verbose, TEXT("[ClearStateRawReport] No have a valid FGameInputDeviceConfiguration in the UGameInputDeveloperSettings. We can't process Raw Input without it. Exiting. (Device %s)"), *UE::GameInput::LexToString(Params.Device));
+		UE_LOGF(LogGameInput, Verbose, "[ClearStateRawReport] No have a valid FGameInputDeviceConfiguration in the UGameInputDeveloperSettings. We can't process Raw Input without it. Exiting. (Device %ls)", *UE::GameInput::LexToString(Params.Device));
 		return;
 	}
 
diff --git a/Engine/Plugins/Runtime/GameInput/Source/GameInputBase/Private/Processors/GameInputDeviceProcessor_Sensor.cpp b/Engine/Plugins/Runtime/GameInput/Source/GameInputBase/Private/Processors/GameInputDeviceProcessor_Sensor.cpp
--- a/Engine/Plugins/Runtime/GameInput/Source/GameInputBase/Private/Processors/GameInputDeviceProcessor_Sensor.cpp
+++ b/Engine/Plugins/Runtime/GameInput/Source/GameInputBase/Private/Processors/GameInputDeviceProcessor_Sensor.cpp
@@ -181,7 +181,7 @@ bool FGameInputGamepadSensorDeviceProcessor::ProcessGamepadSensorState(const FGa
 				Params.InputDeviceId
 			);
 
-		UE_LOG(LogGameInput, VeryVerbose, TEXT("[%hs] OnMotionDetected PlatUser: %d DeviceId %d  MotionSpace: %d"),
+		UE_LOGF(LogGameInput, VeryVerbose, "[%hs] OnMotionDetected PlatUser: %d DeviceId %d  MotionSpace: %d",
 			__func__,
 			Params.PlatformUserId.GetInternalId(),
 			Params.InputDeviceId.GetId(),
@@ -200,7 +200,7 @@ bool FGameInputGamepadSensorDeviceProcessor::ProcessGamepadSensorState(const FGa
 			PreviousState.headingInDegreesFromMagneticNorth,
 			UE::GameInput::SensorMagneticNorthDeadzone);
 
-		UE_LOG(LogGameInput, VeryVerbose, TEXT("[ProcessGamepadSensorState] (Device %s) Degrees From Magnetic North: %.2f (Accuracy: %s) "),
+		UE_LOGF(LogGameInput, VeryVerbose, "[ProcessGamepadSensorState] (Device %ls) Degrees From Magnetic North: %.2f (Accuracy: %ls) ",
 			*UE::GameInput::LexToString(Params.Device), SensorState.headingInDegreesFromMagneticNorth, *UE::GameInput::LexToString(SensorState.headingAccuracy));
 	}
 	
diff --git a/Engine/Plugins/Runtime/GameInput/Source/GameInputBase/Private/Processors/GameInputDeviceProcessor_Touch.cpp b/Engine/Plugins/Runtime/GameInput/Source/GameInputBase/Private/Processors/GameInputDeviceProcessor_Touch.cpp
--- a/Engine/Plugins/Runtime/GameInput/Source/GameInputBase/Private/Processors/GameInputDeviceProcessor_Touch.cpp
+++ b/Engine/Plugins/Runtime/GameInput/Source/GameInputBase/Private/Processors/GameInputDeviceProcessor_Touch.cpp
@@ -45,7 +45,7 @@ bool FGameInputTouchDeviceProcessor::ProcessInput(const FGameInputEventParams& P
 	// We can't tell slate about input messages from an invalid platform user
 	if (!Params.PlatformUserId.IsValid())
 	{
-		UE_LOG(LogGameInput, Verbose, TEXT("[FGameInputTouchDeviceProcessor::EvaluateButtonStates] Attempting to evaluate button states with an invalid platform user id of '%d'. The button messages will not be sent."), Params.PlatformUserId.GetInternalId());
+		UE_LOGF(LogGameInput, Verbose, "[FGameInputTouchDeviceProcessor::EvaluateButtonStates] Attempting to evaluate button states with an invalid platform user id of '%d'. The button messages will not be sent.", Params.PlatformUserId.GetInternalId());
 		return false;
 	}
 
diff --git a/Engine/Plugins/Runtime/GameInput/Source/GameInputBase/Private/Processors/IGameInputDeviceProcessor.cpp b/Engine/Plugins/Runtime/GameInput/Source/GameInputBase/Private/Processors/IGameInputDeviceProcessor.cpp
--- a/Engine/Plugins/Runtime/GameInput/Source/GameInputBase/Private/Processors/IGameInputDeviceProcessor.cpp
+++ b/Engine/Plugins/Runtime/GameInput/Source/GameInputBase/Private/Processors/IGameInputDeviceProcessor.cpp
@@ -42,7 +42,7 @@ void IGameInputDeviceProcessor::OnControllerAnalog(const FGameInputEventParams&
 {
 	if (OldAxisValueNormalized != NewAxisValueNormalized || FMath::Abs(NewAxisValueNormalized) > DeadZone)
 	{		
-		UE_LOG(LogGameInput, VeryVerbose, TEXT("Device %s (PlatformUserId = %d, InputDeviceId = %d) - Analog %s : %.3f"),
+		UE_LOGF(LogGameInput, VeryVerbose, "Device %ls (PlatformUserId = %d, InputDeviceId = %d) - Analog %ls : %.3f",
 			*UE::GameInput::LexToString(Params.Device),
 			Params.PlatformUserId.GetInternalId(),
 			Params.InputDeviceId.GetId(),
@@ -87,7 +87,7 @@ void IGameInputDeviceProcessor::EvaluateButtonStates(
 
 	if (!RepeatTime)
 	{
-		UE_LOG(LogGameInput, Error, TEXT("[IGameInputDeviceProcessor::EvaluateButtonStates] Invalid RepeatTime array given to evaluate button states!"));
+		UE_LOGF(LogGameInput, Error, "[IGameInputDeviceProcessor::EvaluateButtonStates] Invalid RepeatTime array given to evaluate button states!");
 		return;
 	}
 
@@ -100,7 +100,7 @@ void IGameInputDeviceProcessor::EvaluateButtonStates(
 	// We can't tell slate about input messages from an invalid platform user
 	if (!Params.PlatformUserId.IsValid())
 	{
-		UE_LOG(LogGameInput, Verbose, TEXT("[IGameInputDeviceProcessor::EvaluateButtonStates] Attempting to evaluate button states with an invalid platform user id of '%d'. The button messages will not be sent."), Params.PlatformUserId.GetInternalId());
+		UE_LOGF(LogGameInput, Verbose, "[IGameInputDeviceProcessor::EvaluateButtonStates] Attempting to evaluate button states with an invalid platform user id of '%d'. The button messages will not be sent.", Params.PlatformUserId.GetInternalId());
 		return;
 	}
 
@@ -122,13 +122,13 @@ PRAGMA_ENABLE_DEPRECATION_WARNINGS
 
 				if (0 != (CurrentButtonHeldMask & BitMask))
 				{
-					UE_LOG(LogGameInput, Verbose, TEXT("[FGameInputDevice::EvaluateButtonStates] (PlatformUserId = %d, InputDeviceId = %d) - Button '%s' Pressed"), Params.PlatformUserId.GetInternalId(), Params.InputDeviceId.GetId(), *ButtonKey.ToString());
+					UE_LOGF(LogGameInput, Verbose, "[FGameInputDevice::EvaluateButtonStates] (PlatformUserId = %d, InputDeviceId = %d) - Button '%ls' Pressed", Params.PlatformUserId.GetInternalId(), Params.InputDeviceId.GetId(), *ButtonKey.ToString());
 					Params.MessageHandler->OnControllerButtonPressed(ButtonKey, Params.PlatformUserId, Params.InputDeviceId, false);
 					RepeatTime[n] = CurrentTime + InitialButtonRepeatDelay;
 				}
 				else
 				{
-					UE_LOG(LogGameInput, Verbose, TEXT("[FGameInputDevice::EvaluateButtonStates] (PlatformUserId = %d, InputDeviceId = %d) - Button '%s' Released"), Params.PlatformUserId.GetInternalId(), Params.InputDeviceId.GetId(), *ButtonKey.ToString());
+					UE_LOGF(LogGameInput, Verbose, "[FGameInputDevice::EvaluateButtonStates] (PlatformUserId = %d, InputDeviceId = %d) - Button '%ls' Released", Params.PlatformUserId.GetInternalId(), Params.InputDeviceId.GetId(), *ButtonKey.ToString());
 					Params.MessageHandler->OnControllerButtonReleased(ButtonKey, Params.PlatformUserId, Params.InputDeviceId, false);
 					RepeatTime[n] = 0.0;
 				}
@@ -142,7 +142,7 @@ PRAGMA_DISABLE_DEPRECATION_WARNINGS
 PRAGMA_ENABLE_DEPRECATION_WARNINGS
 
 				RepeatTime[n] = CurrentTime + ButtonRepeatDelay;
-				UE_LOG(LogGameInput, Verbose, TEXT("[FGameInputDevice::EvaluateButtonStates] (PlatformUserId = %d, InputDeviceId = %d) - Button '%s' Repeat Pressed"), Params.PlatformUserId.GetInternalId(), Params.InputDeviceId.GetId(), *ButtonKey.ToString());
+				UE_LOGF(LogGameInput, Verbose, "[FGameInputDevice::EvaluateButtonStates] (PlatformUserId = %d, InputDeviceId = %d) - Button '%ls' Repeat Pressed", Params.PlatformUserId.GetInternalId(), Params.InputDeviceId.GetId(), *ButtonKey.ToString());
 				Params.MessageHandler->OnControllerButtonPressed(ButtonKey, Params.PlatformUserId, Params.InputDeviceId, true);
 			}
 		}
@@ -175,7 +175,7 @@ void IGameInputDeviceProcessor::EvaluateSwitchState(
 	// We can't send any slate input events to slate if the platform user is invalid
 	if (!Params.PlatformUserId.IsValid())
 	{
-		UE_LOG(LogGameInput, Verbose, TEXT("[IGameInputDeviceProcessor::EvaluateSwitchState] Attempting to evaluate button states with an invalid platform user id of '%d'. The input messages will not be sent."), Params.PlatformUserId.GetInternalId());
+		UE_LOGF(LogGameInput, Verbose, "[IGameInputDeviceProcessor::EvaluateSwitchState] Attempting to evaluate button states with an invalid platform user id of '%d'. The input messages will not be sent.", Params.PlatformUserId.GetInternalId());
 		return;
 	}
 
@@ -191,7 +191,7 @@ void IGameInputDeviceProcessor::EvaluateSwitchState(
 		{
 			for (const FGamepadKeyNames::Type KeyName : *PrevKeyArray)
 			{
-				UE_LOG(LogGameInput, Verbose, TEXT("[EvaluateSwitchState] (PlatformUserId = %d, InputDeviceId = %d) - Switch '%s' Released"), Params.PlatformUserId.GetInternalId(), Params.InputDeviceId.GetId(), *KeyName.ToString());
+				UE_LOGF(LogGameInput, Verbose, "[EvaluateSwitchState] (PlatformUserId = %d, InputDeviceId = %d) - Switch '%ls' Released", Params.PlatformUserId.GetInternalId(), Params.InputDeviceId.GetId(), *KeyName.ToString());
 				Params.MessageHandler->OnControllerButtonReleased(KeyName, Params.PlatformUserId, Params.InputDeviceId, false);
 			}
 
@@ -203,7 +203,7 @@ void IGameInputDeviceProcessor::EvaluateSwitchState(
 		{
 			for (const FGamepadKeyNames::Type KeyName : *CurrentKeyArray)
 			{
-				UE_LOG(LogGameInput, Verbose, TEXT("[EvaluateSwitchState] (PlatformUserId = %d, InputDeviceId = %d) - Switch '%s' Pressed"), Params.PlatformUserId.GetInternalId(), Params.InputDeviceId.GetId(), *KeyName.ToString());
+				UE_LOGF(LogGameInput, Verbose, "[EvaluateSwitchState] (PlatformUserId = %d, InputDeviceId = %d) - Switch '%ls' Pressed", Params.PlatformUserId.GetInternalId(), Params.InputDeviceId.GetId(), *KeyName.ToString());
 				Params.MessageHandler->OnControllerButtonPressed(KeyName, Params.PlatformUserId, Params.InputDeviceId, false);
 			}
 
@@ -217,7 +217,7 @@ void IGameInputDeviceProcessor::EvaluateSwitchState(
 		{
 			for (const FGamepadKeyNames::Type KeyName : *CurrentKeyArray)
 			{
-				UE_LOG(LogGameInput, Verbose, TEXT("[EvaluateSwitchState] (PlatformUserId = %d, InputDeviceId = %d) - Switch '%s' Repeat Pressed"), Params.PlatformUserId.GetInternalId(), Params.InputDeviceId.GetId(), *KeyName.ToString());
+				UE_LOGF(LogGameInput, Verbose, "[EvaluateSwitchState] (PlatformUserId = %d, InputDeviceId = %d) - Switch '%ls' Repeat Pressed", Params.PlatformUserId.GetInternalId(), Params.InputDeviceId.GetId(), *KeyName.ToString());
 				Params.MessageHandler->OnControllerButtonPressed(KeyName, Params.PlatformUserId, Params.InputDeviceId, true);
 			}
 		}
